Add to import process the ability to templatize for non-imported properties of USoundWaves from pre-existing USoundWave

diff --git a/Engine/Plugins/Editor/FacialAnimation/Source/FacialAnimationEditor/Private/FacialAnimationImportItem.cpp b/Engine/Plugins/Editor/FacialAnimation/Source/FacialAnimationEditor/Private/FacialAnimationImportItem.cpp
--- a/Engine/Plugins/Editor/FacialAnimation/Source/FacialAnimationEditor/Private/FacialAnimationImportItem.cpp
+++ b/Engine/Plugins/Editor/FacialAnimation/Source/FacialAnimationEditor/Private/FacialAnimationImportItem.cpp
@@ -41,7 +41,7 @@ USoundWave* FFacialAnimationImportItem::ImportSoundWave(const FString& InSoundWa
 
 	// We set the correct options in the constructor, so run the import silently
 	USoundFactory* SoundWaveFactory = NewObject<USoundFactory>();
-	SoundWaveFactory->SuppressImportOverwriteDialog();
+	SoundWaveFactory->SuppressImportDialogs();
 
 	// Perform the actual import
 	USoundWave* const SoundWave = ImportObject<USoundWave>(SoundWavePackage, *InSoundWaveAssetName, RF_Public | RF_Standalone, *InWavFilename, nullptr, SoundWaveFactory);
diff --git a/Engine/Source/Editor/AudioEditor/Classes/Factories/SoundFactory.h b/Engine/Source/Editor/AudioEditor/Classes/Factories/SoundFactory.h
--- a/Engine/Source/Editor/AudioEditor/Classes/Factories/SoundFactory.h
+++ b/Engine/Source/Editor/AudioEditor/Classes/Factories/SoundFactory.h
@@ -11,6 +11,9 @@
 #include "Factories/Factory.h"
 #include "SoundFactory.generated.h"
 
+
+class USoundWave;
+
 UCLASS(MinimalAPI, hidecategories=Object)
 class USoundFactory : public UFactory
 {
@@ -40,15 +43,30 @@ class USoundFactory : public UFactory
 	UPROPERTY(EditAnywhere, Category=SoundFactory, meta=(ToolTip="If not empty, generated SoundCues will be placed in PackageCuePackageSuffix, but only if bAutoCreateCue is true"))
 	FString CuePackageSuffix;
 
-
 	//~ Begin UFactory Interface
 	virtual UObject* FactoryCreateBinary( UClass* Class, UObject* InParent, FName Name, EObjectFlags Flags, UObject* Context, const TCHAR* FileType, const uint8*& Buffer, const uint8* BufferEnd, FFeedbackContext* Warn ) override;
 	//~ End UFactory Interface
 	
-	/** Suppresses the import overwrite dialog until one iteration of FactoryCreateBinary completes; this is primarily used for reimporting sounds */
-	AUDIOEDITOR_API static void SuppressImportOverwriteDialog();
+	/** Suppresses all dialogs pertaining to importing using factory */
+	void SuppressImportDialogs();
+
+	virtual void CleanUp() override;
+
+protected:
+	enum ESuppressImportDialog
+	{
+		None = 0,
+		Overwrite = 1 << 0,
+		UseTemplate = 1 << 1
+	};
+
+	uint8 SuppressImportDialogOptions;
 
 private:
+	void UpdateTemplate();
+
+	TWeakObjectPtr<USoundWave> TemplateSoundWave;
+
 
 	UObject* CreateObject(UClass* Class, UObject* InParent, FName Name, EObjectFlags Flags, UObject* Context, const TCHAR* FileType, const uint8*& Buffer, const uint8* BufferEnd, FFeedbackContext* Warn);
 
diff --git a/Engine/Source/Editor/AudioEditor/Private/AudioEditorModule.cpp b/Engine/Source/Editor/AudioEditor/Private/AudioEditorModule.cpp
--- a/Engine/Source/Editor/AudioEditor/Private/AudioEditorModule.cpp
+++ b/Engine/Source/Editor/AudioEditor/Private/AudioEditorModule.cpp
@@ -280,7 +280,7 @@ public:
 
 		// Setup sane defaults for importing localized sound waves
 		SoundWaveFactory->bAutoCreateCue = false;
-		SoundWaveFactory->SuppressImportOverwriteDialog();
+		SoundWaveFactory->SuppressImportDialogs();
 
 		return ImportObject<USoundWave>(SoundWavePackage, *InSoundWaveAssetName, RF_Public | RF_Standalone, *InWavFilename, nullptr, SoundWaveFactory);
 	}
diff --git a/Engine/Source/Editor/AudioEditor/Private/Factories/ReimportSoundFactory.cpp b/Engine/Source/Editor/AudioEditor/Private/Factories/ReimportSoundFactory.cpp
--- a/Engine/Source/Editor/AudioEditor/Private/Factories/ReimportSoundFactory.cpp
+++ b/Engine/Source/Editor/AudioEditor/Private/Factories/ReimportSoundFactory.cpp
@@ -12,7 +12,6 @@
 UReimportSoundFactory::UReimportSoundFactory(const FObjectInitializer& ObjectInitializer)
 	: Super(ObjectInitializer)
 {
-
 	SupportedClass = USoundWave::StaticClass();
 	Formats.Add(TEXT("wav;Wave audio file"));
 
@@ -114,8 +113,8 @@ EReimportResult::Type UReimportSoundFactory::Reimport(UObject* Obj)
 		return EReimportResult::Failed;
 	}
 
-	// Suppress the import overwrite dialog, we want to keep existing settings when re-importing
-	USoundFactory::SuppressImportOverwriteDialog();
+	// Always suppress dialogs on re-import
+	SuppressImportDialogs();
 
 	bool OutCanceled = false;
 	if (!ImportObject(SoundWave->GetClass(), SoundWave->GetOuter(), *SoundWave->GetName(), RF_Public | RF_Standalone, Filename, nullptr, OutCanceled))
diff --git a/Engine/Source/Editor/AudioEditor/Private/Factories/SoundFactory.cpp b/Engine/Source/Editor/AudioEditor/Private/Factories/SoundFactory.cpp
--- a/Engine/Source/Editor/AudioEditor/Private/Factories/SoundFactory.cpp
+++ b/Engine/Source/Editor/AudioEditor/Private/Factories/SoundFactory.cpp
@@ -2,14 +2,17 @@
 
 #include "Factories/SoundFactory.h"
 #include "AssetRegistryModule.h"
-#include "Components/AudioComponent.h"
 #include "Audio.h"
+#include "Components/AudioComponent.h"
+#include "ContentBrowserModule.h"
+#include "IContentBrowserSingleton.h"
+#include "Modules/ModuleManager.h"
 #include "Sound/SoundCue.h"
-#include "Sound/SoundWave.h"
 #include "Sound/SoundNode.h"
 #include "Sound/SoundNodeWavePlayer.h"
 #include "Sound/SoundNodeModulator.h"
 #include "Sound/SoundNodeAttenuation.h"
+#include "Sound/SoundWave.h"
 #include "AudioDeviceManager.h"
 #include "AudioDevice.h"
 #include "AudioEditorModule.h"
@@ -17,84 +20,85 @@
 #include "Misc/MessageDialog.h"
 #include "Misc/FeedbackContext.h"
 #include "EditorFramework/AssetImportData.h"
-#include "AudioEditorModule.h"
-#include "AudioDevice.h"
 #include "SoundFileIO/SoundFileIO.h"
 
 
-static bool bSoundFactorySuppressImportOverwriteDialog = false;
-
-static void InsertSoundNode(USoundCue* SoundCue, UClass* NodeClass, int32 NodeIndex)
+namespace
 {
-	USoundNode* SoundNode = SoundCue->ConstructSoundNode<USoundNode>(NodeClass);
-
-	// If this node allows >0 children but by default has zero - create a connector for starters
-	if (SoundNode->GetMaxChildNodes() > 0 && SoundNode->ChildNodes.Num() == 0)
+	void InsertSoundNode(USoundCue* SoundCue, UClass* NodeClass, int32 NodeIndex)
 	{
-		SoundNode->CreateStartingConnectors();
-	}
+		USoundNode* SoundNode = SoundCue->ConstructSoundNode<USoundNode>(NodeClass);
 
-	SoundNode->GraphNode->NodePosX = -150 * NodeIndex - 100;
-	SoundNode->GraphNode->NodePosY = -35;
+		// If this node allows >0 children but by default has zero - create a connector for starters
+		if (SoundNode->GetMaxChildNodes() > 0 && SoundNode->ChildNodes.Num() == 0)
+		{
+			SoundNode->CreateStartingConnectors();
+		}
 
-	// Link the node to the cue.
-	SoundNode->ChildNodes[0] = SoundCue->FirstNode;
+		SoundNode->GraphNode->NodePosX = -150 * NodeIndex - 100;
+		SoundNode->GraphNode->NodePosY = -35;
 
-	// Link the attenuation node to root.
-	SoundCue->FirstNode = SoundNode;
+		// Link the node to the cue.
+		SoundNode->ChildNodes[0] = SoundCue->FirstNode;
 
-	SoundCue->LinkGraphNodesFromSoundNodes();
-}
+		// Link the attenuation node to root.
+		SoundCue->FirstNode = SoundNode;
 
-static void CreateSoundCue(USoundWave* Sound, UObject* InParent, EObjectFlags Flags, bool bIncludeAttenuationNode, bool bIncludeModulatorNode, bool bIncludeLoopingNode, float CueVolume)
-{
-	// then first create the actual sound cue
-	FString SoundCueName = FString::Printf(TEXT("%s_Cue"), *Sound->GetName());
+		SoundCue->LinkGraphNodesFromSoundNodes();
+	}
 
-	// Create sound cue and wave player
-	USoundCue* SoundCue = NewObject<USoundCue>(InParent, *SoundCueName, Flags);
-	USoundNodeWavePlayer* WavePlayer = SoundCue->ConstructSoundNode<USoundNodeWavePlayer>();
+	void CreateSoundCue(USoundWave* Sound, UObject* InParent, EObjectFlags Flags, bool bIncludeAttenuationNode, bool bIncludeModulatorNode, bool bIncludeLoopingNode, float CueVolume)
+	{
+		// then first create the actual sound cue
+		FString SoundCueName = FString::Printf(TEXT("%s_Cue"), *Sound->GetName());
 
-	int32 NodeIndex = (int32)bIncludeAttenuationNode + (int32)bIncludeModulatorNode + (int32)bIncludeLoopingNode;
+		// Create sound cue and wave player
+		USoundCue* SoundCue = NewObject<USoundCue>(InParent, *SoundCueName, Flags);
+		USoundNodeWavePlayer* WavePlayer = SoundCue->ConstructSoundNode<USoundNodeWavePlayer>();
 
-	WavePlayer->GraphNode->NodePosX = -150 * NodeIndex - 100;
-	WavePlayer->GraphNode->NodePosY = -35;
+		int32 NodeIndex = (int32)bIncludeAttenuationNode + (int32)bIncludeModulatorNode + (int32)bIncludeLoopingNode;
 
-	// Apply the initial volume.
-	SoundCue->VolumeMultiplier = CueVolume;
+		WavePlayer->GraphNode->NodePosX = -150 * NodeIndex - 100;
+		WavePlayer->GraphNode->NodePosY = -35;
 
-	WavePlayer->SetSoundWave(Sound);
-	SoundCue->FirstNode = WavePlayer;
-	SoundCue->LinkGraphNodesFromSoundNodes();
+		// Apply the initial volume.
+		SoundCue->VolumeMultiplier = CueVolume;
 
-	if (bIncludeLoopingNode)
-	{
-		WavePlayer->bLooping = true;
-	}
+		WavePlayer->SetSoundWave(Sound);
+		SoundCue->FirstNode = WavePlayer;
+		SoundCue->LinkGraphNodesFromSoundNodes();
 
-	if (bIncludeModulatorNode)
-	{
-		InsertSoundNode(SoundCue, USoundNodeModulator::StaticClass(), --NodeIndex);
-	}
+		if (bIncludeLoopingNode)
+		{
+			WavePlayer->bLooping = true;
+		}
 
-	if (bIncludeAttenuationNode)
-	{
-		InsertSoundNode(SoundCue, USoundNodeAttenuation::StaticClass(), --NodeIndex);
-	}
+		if (bIncludeModulatorNode)
+		{
+			InsertSoundNode(SoundCue, USoundNodeModulator::StaticClass(), --NodeIndex);
+		}
 
-	// Make sure the content browser finds out about this newly-created object.  This is necessary when sound
-	// cues are created automatically after creating a sound node wave.  See use of bAutoCreateCue in USoundTTSFactory.
-	if ((Flags & (RF_Public | RF_Standalone)) != 0)
-	{
-		// Notify the asset registry
-		FAssetRegistryModule::AssetCreated(SoundCue);
-	}
-}
+		if (bIncludeAttenuationNode)
+		{
+			InsertSoundNode(SoundCue, USoundNodeAttenuation::StaticClass(), --NodeIndex);
+		}
 
+		// Make sure the content browser finds out about this newly-created object.  This is necessary when sound
+		// cues are created automatically after creating a sound node wave.  See use of bAutoCreateCue in USoundTTSFactory.
+		if ((Flags & (RF_Public | RF_Standalone)) != 0)
+		{
+			// Notify the asset registry
+			FAssetRegistryModule::AssetCreated(SoundCue);
+		}
+	}
+} // namespace <>
 
 USoundFactory::USoundFactory(const FObjectInitializer& ObjectInitializer)
 	: Super(ObjectInitializer)
 {
+	SuppressImportDialogOptions = ESuppressImportDialog::None;
+	TemplateSoundWave = nullptr;
+
 	SupportedClass = USoundWave::StaticClass();
 	Formats.Add(TEXT("wav;Wave Audio File"));
 
@@ -250,10 +254,12 @@ UObject* USoundFactory::CreateObject
 			}
 		}
 
-		bool bUseExistingSettings = bSoundFactorySuppressImportOverwriteDialog;
+		UpdateTemplate();
 
-		if (ExistingSound && !bSoundFactorySuppressImportOverwriteDialog && !GIsAutomationTesting)
+		bool bUseExistingSettings = SuppressImportDialogOptions & ESuppressImportDialog::Overwrite;
+		if (ExistingSound && !bUseExistingSettings && !GIsAutomationTesting)
 		{
+			SuppressImportDialogOptions |= ESuppressImportDialog::Overwrite;
 			DisplayOverwriteOptionsDialog(FText::Format(
 				NSLOCTEXT("SoundFactory", "ImportOverwriteWarning", "You are about to import '{0}' over an existing sound."),
 				FText::FromName(Name)));
@@ -292,9 +298,6 @@ UObject* USoundFactory::CreateObject
 		bool bIsAmbiX = (AmbiXTag == TEXT("_ambix"));
 		bool bIsFuMa = (FuMaTag == TEXT("_fuma"));
 
-		// Reset the flag back to false so subsequent imports are not suppressed unless the code explicitly suppresses it
-		bSoundFactorySuppressImportOverwriteDialog = false;
-
 		TArray<uint8> RawWaveData;
 		RawWaveData.Empty(BufferEnd - Buffer);
 		RawWaveData.AddUninitialized(BufferEnd - Buffer);
@@ -334,7 +337,15 @@ UObject* USoundFactory::CreateObject
 
 		// Use pre-existing sound if it exists and we want to keep settings,
 		// otherwise create new sound and import raw data.
-		USoundWave* Sound = (bUseExistingSettings && ExistingSound) ? ExistingSound : NewObject<USoundWave>(InParent, Name, Flags);
+		USoundWave* Sound = (bUseExistingSettings && ExistingSound) ? ExistingSound : NewObject<USoundWave>(InParent, Name, Flags, TemplateSoundWave.Get());
+
+		// These get wiped in PostInitProperties by defaults set from Audio Settings,
+		// so set back to template in this specialized case
+		if (TemplateSoundWave.IsValid())
+		{
+			Sound->SoundClassObject = TemplateSoundWave->SoundClassObject;
+			Sound->SoundConcurrencySettings = TemplateSoundWave->SoundConcurrencySettings;
+		}
 
 		if (bUseExistingSettings && ExistingSound)
 		{
@@ -509,3 +520,44 @@ UObject* USoundFactory::CreateObject
 
 	return nullptr;
 }
+
+void USoundFactory::SuppressImportOverwriteDialog()
+{
+	bSoundFactorySuppressImportOverwriteDialog = true;
+void USoundFactory::SuppressImportDialogs()
+{
+	SuppressImportDialogOptions = ESuppressImportDialog::Overwrite & ESuppressImportDialog::UseTemplate;
+}
+
+void USoundFactory::UpdateTemplate()
+{
+	if (!IsAutomatedImport() && !TemplateSoundWave.IsValid() && !(SuppressImportDialogOptions & ESuppressImportDialog::UseTemplate))
+	{
+		SuppressImportDialogOptions |= ESuppressImportDialog::UseTemplate;
+
+		FContentBrowserModule& ContentBrowserModule = FModuleManager::LoadModuleChecked<FContentBrowserModule>("ContentBrowser");
+		TArray<FAssetData> SelectedAssets;
+		ContentBrowserModule.Get().GetSelectedAssets(SelectedAssets);
+
+		if (SelectedAssets.Num() == 1)
+		{
+			if (USoundWave* SoundWave = Cast<USoundWave>(SelectedAssets[0].GetAsset()))
+			{
+				const bool bUseTemplateSoundWave = FMessageDialog::Open(EAppMsgType::YesNo, FText::Format(
+					NSLOCTEXT("SoundFactory", "UseSoundWaveTemplate", "Use the selected Sound Wave '{0}' in the Content Browser as a template for sound(s) being imported?"),
+					FText::FromString(SoundWave->GetName()))) == EAppReturnType::Yes;
+
+				if (bUseTemplateSoundWave)
+				{
+					TemplateSoundWave = SoundWave;
+				}
+			}
+		}
+	}
+}
+
+void USoundFactory::CleanUp()
+{
+	SuppressImportDialogOptions = ESuppressImportDialog::None;
+	TemplateSoundWave.Reset();
+}
diff --git a/Engine/Source/Runtime/Engine/Classes/Sound/SoundBase.h b/Engine/Source/Runtime/Engine/Classes/Sound/SoundBase.h
--- a/Engine/Source/Runtime/Engine/Classes/Sound/SoundBase.h
+++ b/Engine/Source/Runtime/Engine/Classes/Sound/SoundBase.h
@@ -36,13 +36,10 @@ public:
 	static USoundClass* DefaultSoundClassObject;
 	static USoundConcurrency* DefaultSoundConcurrencyObject;
 
-protected:
 	/** Sound class this sound belongs to */
-	UPROPERTY(EditAnywhere, Category=Sound, meta=(DisplayName = "Sound Class"), AssetRegistrySearchable)
+	UPROPERTY(EditAnywhere, Category = Sound, meta = (DisplayName = "Sound Class"), AssetRegistrySearchable)
 	USoundClass* SoundClassObject;
 
-public:
-
 	/** When "stat sounds -debug" has been specified, draw this sound's attenuation shape when the sound is audible. For debugging purpose only. */
 	UPROPERTY(EditAnywhere, Category = Debug)
 	uint8 bDebug:1;
