Fix merge errors
#rnx

diff --git a/Engine/Source/Editor/AudioEditor/Private/Factories/SoundFactory.cpp b/Engine/Source/Editor/AudioEditor/Private/Factories/SoundFactory.cpp
--- a/Engine/Source/Editor/AudioEditor/Private/Factories/SoundFactory.cpp
+++ b/Engine/Source/Editor/AudioEditor/Private/Factories/SoundFactory.cpp
@@ -135,7 +135,6 @@ UObject* USoundFactory::FactoryCreateBinary
 
 	UObject* SoundObject = nullptr;
 
-	const bool SuppressOverwrite = bSoundFactorySuppressImportOverwriteDialog;
 	if (FCString::Stricmp(FileType, TEXT("WAV")) == 0)
 	{
 		SoundObject = CreateObject(Class, InParent, Name, Flags, Context, FileType, Buffer, BufferEnd, Warn);
@@ -158,7 +157,6 @@ UObject* USoundFactory::FactoryCreateBinary
 
 			// Perpetuate the setting of the suppression flag to avoid
 			// user notification if we attempt to call CreateObject twice
-			bSoundFactorySuppressImportOverwriteDialog = SuppressOverwrite;
 			SoundObject = CreateObject(Class, InParent, Name, Flags, Context, TEXT("WAV"), Ptr, Ptr + RawWaveData.Num() - 1, Warn);
 		}
 	}
@@ -174,11 +172,6 @@ UObject* USoundFactory::FactoryCreateBinary
 	return SoundObject;
 }
 
-void USoundFactory::SuppressImportOverwriteDialog()
-{
-	bSoundFactorySuppressImportOverwriteDialog = true;
-}
-
 UObject* USoundFactory::CreateObject
 (
 	UClass*				Class,
@@ -344,7 +337,7 @@ UObject* USoundFactory::CreateObject
 		if (TemplateSoundWave.IsValid())
 		{
 			Sound->SoundClassObject = TemplateSoundWave->SoundClassObject;
-			Sound->SoundConcurrencySettings = TemplateSoundWave->SoundConcurrencySettings;
+			Sound->ConcurrencySet = TemplateSoundWave->ConcurrencySet;
 		}
 
 		if (bUseExistingSettings && ExistingSound)
@@ -521,9 +514,6 @@ UObject* USoundFactory::CreateObject
 	return nullptr;
 }
 
-void USoundFactory::SuppressImportOverwriteDialog()
-{
-	bSoundFactorySuppressImportOverwriteDialog = true;
 void USoundFactory::SuppressImportDialogs()
 {
 	SuppressImportDialogOptions = ESuppressImportDialog::Overwrite & ESuppressImportDialog::UseTemplate;
