Fix compiler error from bad merge fall out due to missing, non-public implementation

diff --git a/Engine/Source/Editor/AudioEditor/Classes/Factories/SoundFactory.h b/Engine/Source/Editor/AudioEditor/Classes/Factories/SoundFactory.h
--- a/Engine/Source/Editor/AudioEditor/Classes/Factories/SoundFactory.h
+++ b/Engine/Source/Editor/AudioEditor/Classes/Factories/SoundFactory.h
@@ -46,9 +46,15 @@ class USoundFactory : public UFactory
 	//~ Begin UFactory Interface
 	virtual UObject* FactoryCreateBinary( UClass* Class, UObject* InParent, FName Name, EObjectFlags Flags, UObject* Context, const TCHAR* FileType, const uint8*& Buffer, const uint8* BufferEnd, FFeedbackContext* Warn ) override;
 	//~ End UFactory Interface
-	
+
+	uint8 SuppressImportDialogOptions;
+
+public:
 	/** Suppresses all dialogs pertaining to importing using factory */
-	void SuppressImportDialogs();
+	void SuppressImportDialogs()
+	{
+		SuppressImportDialogOptions = ESuppressImportDialog::Overwrite & ESuppressImportDialog::UseTemplate;
+	}
 
 	virtual void CleanUp() override;
 
@@ -60,8 +66,6 @@ protected:
 		UseTemplate = 1 << 1
 	};
 
-	uint8 SuppressImportDialogOptions;
-
 private:
 	void UpdateTemplate();
 
@@ -69,7 +73,4 @@ private:
 
 
 	UObject* CreateObject(UClass* Class, UObject* InParent, FName Name, EObjectFlags Flags, UObject* Context, const TCHAR* FileType, const uint8*& Buffer, const uint8* BufferEnd, FFeedbackContext* Warn);
-
-	/** If true, the overwrite dialog should not be shown while importing */
-	static bool bSuppressImportOverwriteDialog;
 };
