Missing file

diff --git a/Engine/Source/Editor/AudioEditor/Private/Factories/SoundFactory.cpp b/Engine/Source/Editor/AudioEditor/Private/Factories/SoundFactory.cpp
--- a/Engine/Source/Editor/AudioEditor/Private/Factories/SoundFactory.cpp
+++ b/Engine/Source/Editor/AudioEditor/Private/Factories/SoundFactory.cpp
@@ -514,11 +514,6 @@ UObject* USoundFactory::CreateObject
 	return nullptr;
 }
 
-void USoundFactory::SuppressImportDialogs()
-{
-	SuppressImportDialogOptions = ESuppressImportDialog::Overwrite & ESuppressImportDialog::UseTemplate;
-}
-
 void USoundFactory::UpdateTemplate()
 {
 	if (!IsAutomatedImport() && !TemplateSoundWave.IsValid() && !(SuppressImportDialogOptions & ESuppressImportDialog::UseTemplate))
