Fix compile error

diff --git a/Engine/Source/Editor/AudioEditor/Classes/Factories/SoundFactory.h b/Engine/Source/Editor/AudioEditor/Classes/Factories/SoundFactory.h
--- a/Engine/Source/Editor/AudioEditor/Classes/Factories/SoundFactory.h
+++ b/Engine/Source/Editor/AudioEditor/Classes/Factories/SoundFactory.h
@@ -51,10 +51,7 @@ class USoundFactory : public UFactory
 
 public:
 	/** Suppresses all dialogs pertaining to importing using factory */
-	void SuppressImportDialogs()
-	{
-		SuppressImportDialogOptions = ESuppressImportDialog::Overwrite & ESuppressImportDialog::UseTemplate;
-	}
+	AUDIOEDITOR_API void SuppressImportDialogs();
 
 	virtual void CleanUp() override;
 
