Fix merge conflict (again) #rb none #jira none #okforgithub public

diff --git a/Engine/Source/Editor/AudioEditor/Private/Factories/SoundFactory.cpp b/Engine/Source/Editor/AudioEditor/Private/Factories/SoundFactory.cpp
--- a/Engine/Source/Editor/AudioEditor/Private/Factories/SoundFactory.cpp
+++ b/Engine/Source/Editor/AudioEditor/Private/Factories/SoundFactory.cpp
@@ -512,6 +512,10 @@ UObject* USoundFactory::CreateObject
 	}
 
 	return nullptr;
+
+void USoundFactory::SuppressImportDialogs()
+{
+	SuppressImportDialogOptions = ESuppressImportDialog::Overwrite & ESuppressImportDialog::UseTemplate;
 }
 
 void USoundFactory::UpdateTemplate()
