Missing }

diff --git a/Engine/Source/Editor/AudioEditor/Private/Factories/SoundFactory.cpp b/Engine/Source/Editor/AudioEditor/Private/Factories/SoundFactory.cpp
--- a/Engine/Source/Editor/AudioEditor/Private/Factories/SoundFactory.cpp
+++ b/Engine/Source/Editor/AudioEditor/Private/Factories/SoundFactory.cpp
@@ -512,6 +512,7 @@ UObject* USoundFactory::CreateObject
 	}
 
 	return nullptr;
+}
 
 void USoundFactory::SuppressImportDialogs()
 {
