[NO NON-PII COMMIT MESSAGE]

diff --git a/Engine/Source/Editor/AudioEditor/Private/AudioEditorModule.cpp b/Engine/Source/Editor/AudioEditor/Private/AudioEditorModule.cpp
--- a/Engine/Source/Editor/AudioEditor/Private/AudioEditorModule.cpp
+++ b/Engine/Source/Editor/AudioEditor/Private/AudioEditorModule.cpp
@@ -94,7 +94,7 @@ public:
 #if WITH_SNDFILE_IO
 		if (!Audio::InitSoundFileIOManager())
 		{
-			UE_LOG(LogAudioEditor, Error, TEXT("LibSoundFile failed to load. Importing audio will not work correctly."));
+			UE_LOG(LogAudioEditor, Display, TEXT("LibSoundFile failed to load. Importing audio will not work correctly."));
 		}
 #endif // WITH_SNDFILE_IO
 	}
