Editor clean-up of descriptions/class friendly names

diff --git a/Engine/Source/Editor/AudioEditor/Private/AssetTypeActions/AssetTypeActions_SoundEffectPreset.h b/Engine/Source/Editor/AudioEditor/Private/AssetTypeActions/AssetTypeActions_SoundEffectPreset.h
--- a/Engine/Source/Editor/AudioEditor/Private/AssetTypeActions/AssetTypeActions_SoundEffectPreset.h
+++ b/Engine/Source/Editor/AudioEditor/Private/AssetTypeActions/AssetTypeActions_SoundEffectPreset.h
@@ -18,7 +18,7 @@ class FAssetTypeActions_SoundEffectSubmixPreset : public FAssetTypeActions_Base
 {
 public:
 	// IAssetTypeActions Implementation
-	virtual FText GetName() const override { return NSLOCTEXT("AssetTypeActions", "AssetTypeActions_SoundSubmixPreset", "Submix Effect preset"); }
+	virtual FText GetName() const override { return NSLOCTEXT("AssetTypeActions", "AssetTypeActions_SoundSubmixPreset", "Submix Effect Preset"); }
 	virtual FColor GetTypeColor() const override { return FColor(99, 63, 56); }
 	virtual UClass* GetSupportedClass() const override { return USoundEffectSubmixPreset::StaticClass(); }
 	virtual uint32 GetCategories() override { return EAssetTypeCategories::Sounds; }
@@ -44,8 +44,6 @@ public:
 	virtual uint32 GetCategories() override { return EAssetTypeCategories::Sounds; }
 };
 
-
-
 class FAssetTypeActions_SoundEffectPreset : public FAssetTypeActions_Base
 {
 public:
@@ -59,6 +57,5 @@ public:
 	//~ End FAssetTypeActions_Base
 
 private:
-
 	USoundEffectPreset* EffectPreset;
 };
\ No newline at end of file
diff --git a/Engine/Source/Runtime/Engine/Classes/Sound/SoundEffectSource.h b/Engine/Source/Runtime/Engine/Classes/Sound/SoundEffectSource.h
--- a/Engine/Source/Runtime/Engine/Classes/Sound/SoundEffectSource.h
+++ b/Engine/Source/Runtime/Engine/Classes/Sound/SoundEffectSource.h
@@ -12,7 +12,7 @@
 class FSoundEffectSource;
 class FSoundEffectBase;
 
-/** This is here to make sure users don't mix up source and submix effects in the editor. Asset sorting, drag-n-drop, etc. */
+/** Preset of a source effect that can be shared between chains. */
 UCLASS(config = Engine, abstract, editinlinenew, BlueprintType)
 class ENGINE_API USoundEffectSourcePreset : public USoundEffectPreset
 {
@@ -31,7 +31,7 @@ struct ENGINE_API FSourceEffectChainEntry
 	uint32 bBypass : 1;
 };
 
-
+/** Chain of source effect presets that can be shared between referencing sounds. */
 UCLASS(BlueprintType)
 class ENGINE_API USoundEffectSourcePresetChain : public UObject
 {
diff --git a/Engine/Source/Runtime/Engine/Classes/Sound/SoundEffectSubmix.h b/Engine/Source/Runtime/Engine/Classes/Sound/SoundEffectSubmix.h
--- a/Engine/Source/Runtime/Engine/Classes/Sound/SoundEffectSubmix.h
+++ b/Engine/Source/Runtime/Engine/Classes/Sound/SoundEffectSubmix.h
@@ -10,7 +10,7 @@
 class FSoundEffectSubmix;
 
 
-/** This is here to make sure users don't mix up source and submix effects in the editor. Asset sorting, drag-n-drop, etc. */
+/** Preset of a submix effect that can be shared between sounds. */
 UCLASS(config = Engine, hidecategories = Object, abstract, editinlinenew, BlueprintType)
 class ENGINE_API USoundEffectSubmixPreset : public USoundEffectPreset
 {
