Fix for prompts not getting supressed when importing/re-importing sound.

diff --git a/Engine/Source/Editor/AudioEditor/Private/Factories/SoundFactory.cpp b/Engine/Source/Editor/AudioEditor/Private/Factories/SoundFactory.cpp
--- a/Engine/Source/Editor/AudioEditor/Private/Factories/SoundFactory.cpp
+++ b/Engine/Source/Editor/AudioEditor/Private/Factories/SoundFactory.cpp
@@ -521,7 +521,7 @@ UObject* USoundFactory::CreateObject
 
 void USoundFactory::SuppressImportDialogs()
 {
-	SuppressImportDialogOptions = ESuppressImportDialog::Overwrite & ESuppressImportDialog::UseTemplate;
+	SuppressImportDialogOptions = ESuppressImportDialog::Overwrite | ESuppressImportDialog::UseTemplate;
 }
 
 void USoundFactory::UpdateTemplate()
