Add alignment/straighten tools to SoundCue editor

diff --git a/Engine/Source/Editor/AudioEditor/Private/SoundCueEditor.cpp b/Engine/Source/Editor/AudioEditor/Private/SoundCueEditor.cpp
--- a/Engine/Source/Editor/AudioEditor/Private/SoundCueEditor.cpp
+++ b/Engine/Source/Editor/AudioEditor/Private/SoundCueEditor.cpp
@@ -572,6 +572,44 @@ TSharedRef<SGraphEditor> FSoundCueEditor::CreateGraphEditorWidget()
 			FExecuteAction::CreateSP( this, &FSoundCueEditor::DuplicateNodes ),
 			FCanExecuteAction::CreateSP( this, &FSoundCueEditor::CanDuplicateNodes )
 			);
+
+		// Alignment Commands
+		GraphEditorCommands->MapAction(FGraphEditorCommands::Get().AlignNodesTop,
+			FExecuteAction::CreateSP(this, &FSoundCueEditor::OnAlignTop)
+		);
+
+		GraphEditorCommands->MapAction(FGraphEditorCommands::Get().AlignNodesMiddle,
+			FExecuteAction::CreateSP(this, &FSoundCueEditor::OnAlignMiddle)
+		);
+
+		GraphEditorCommands->MapAction(FGraphEditorCommands::Get().AlignNodesBottom,
+			FExecuteAction::CreateSP(this, &FSoundCueEditor::OnAlignBottom)
+		);
+
+		GraphEditorCommands->MapAction(FGraphEditorCommands::Get().AlignNodesLeft,
+			FExecuteAction::CreateSP(this, &FSoundCueEditor::OnAlignLeft)
+		);
+
+		GraphEditorCommands->MapAction(FGraphEditorCommands::Get().AlignNodesCenter,
+			FExecuteAction::CreateSP(this, &FSoundCueEditor::OnAlignCenter)
+		);
+
+		GraphEditorCommands->MapAction(FGraphEditorCommands::Get().AlignNodesRight,
+			FExecuteAction::CreateSP(this, &FSoundCueEditor::OnAlignRight)
+		);
+
+		GraphEditorCommands->MapAction(FGraphEditorCommands::Get().StraightenConnections,
+			FExecuteAction::CreateSP(this, &FSoundCueEditor::OnStraightenConnections)
+		);
+
+		// Distribution Commands
+		GraphEditorCommands->MapAction(FGraphEditorCommands::Get().DistributeNodesHorizontally,
+			FExecuteAction::CreateSP(this, &FSoundCueEditor::OnDistributeNodesH)
+		);
+
+		GraphEditorCommands->MapAction(FGraphEditorCommands::Get().DistributeNodesVertically,
+			FExecuteAction::CreateSP(this, &FSoundCueEditor::OnDistributeNodesV)
+		);
 	}
 
 	FGraphAppearanceInfo AppearanceInfo;
@@ -905,4 +943,76 @@ void FSoundCueEditor::RedoGraphAction()
 	GEditor->RedoTransaction();
 }
 
+void FSoundCueEditor::OnAlignTop()
+{
+	if (SoundCueGraphEditor.IsValid())
+	{
+		SoundCueGraphEditor->OnAlignTop();
+	}
+}
+
+void FSoundCueEditor::OnAlignMiddle()
+{
+	if (SoundCueGraphEditor.IsValid())
+	{
+		SoundCueGraphEditor->OnAlignMiddle();
+	}
+}
+
+void FSoundCueEditor::OnAlignBottom()
+{
+	if (SoundCueGraphEditor.IsValid())
+	{
+		SoundCueGraphEditor->OnAlignBottom();
+	}
+}
+
+void FSoundCueEditor::OnAlignLeft()
+{
+	if (SoundCueGraphEditor.IsValid())
+	{
+		SoundCueGraphEditor->OnAlignLeft();
+	}
+}
+
+void FSoundCueEditor::OnAlignCenter()
+{
+	if (SoundCueGraphEditor.IsValid())
+	{
+		SoundCueGraphEditor->OnAlignCenter();
+	}
+}
+
+void FSoundCueEditor::OnAlignRight()
+{
+	if (SoundCueGraphEditor.IsValid())
+	{
+		SoundCueGraphEditor->OnAlignRight();
+	}
+}
+
+void FSoundCueEditor::OnStraightenConnections()
+{
+	if (SoundCueGraphEditor.IsValid())
+	{
+		SoundCueGraphEditor->OnStraightenConnections();
+	}
+}
+
+void FSoundCueEditor::OnDistributeNodesH()
+{
+	if (SoundCueGraphEditor.IsValid())
+	{
+		SoundCueGraphEditor->OnDistributeNodesH();
+	}
+}
+
+void FSoundCueEditor::OnDistributeNodesV()
+{
+	if (SoundCueGraphEditor.IsValid())
+	{
+		SoundCueGraphEditor->OnDistributeNodesV();
+	}
+}
+
 #undef LOCTEXT_NAMESPACE
diff --git a/Engine/Source/Editor/AudioEditor/Private/SoundCueEditor.h b/Engine/Source/Editor/AudioEditor/Private/SoundCueEditor.h
--- a/Engine/Source/Editor/AudioEditor/Private/SoundCueEditor.h
+++ b/Engine/Source/Editor/AudioEditor/Private/SoundCueEditor.h
@@ -42,7 +42,7 @@ public:
 	virtual bool GetBoundsForSelectedNodes(class FSlateRect& Rect, float Padding) override;
 	virtual int32 GetNumberOfSelectedNodes() const override;
 	virtual TSet<UObject*> GetSelectedNodes() const override;
-	
+
 	/** IToolkit interface */
 	virtual FName GetToolkitFName() const override;
 	virtual FText GetBaseToolkitName() const override;
@@ -123,6 +123,18 @@ protected:
 	/** Called to redo the last undone action */
 	void RedoGraphAction();
 
+	void OnAlignTop();
+	void OnAlignMiddle();
+	void OnAlignBottom();
+	void OnAlignLeft();
+	void OnAlignCenter();
+	void OnAlignRight();
+
+	void OnStraightenConnections();
+
+	void OnDistributeNodesH();
+	void OnDistributeNodesV();
+
 private:
 	/** FNotifyHook interface */
 	virtual void NotifyPostChange( const FPropertyChangedEvent& PropertyChangedEvent, UProperty* PropertyThatChanged) override;
@@ -132,7 +144,7 @@ private:
 
 	/** Builds the toolbar widget for the SoundCue editor */
 	void ExtendToolbar();
-	
+
 	/** Binds new graph commands to delegates */
 	void BindGraphCommands();
 
diff --git a/Engine/Source/Editor/AudioEditor/Private/SoundCueGraphNode.cpp b/Engine/Source/Editor/AudioEditor/Private/SoundCueGraphNode.cpp
--- a/Engine/Source/Editor/AudioEditor/Private/SoundCueGraphNode.cpp
+++ b/Engine/Source/Editor/AudioEditor/Private/SoundCueGraphNode.cpp
@@ -223,6 +223,24 @@ void USoundCueGraphNode::GetContextMenuActions(const FGraphNodeContextMenuBuilde
 	}
 	else if (Context.Node)
 	{
+		Context.MenuBuilder->AddSubMenu(LOCTEXT("AlignmentHeader", "Alignment"), FText(), FNewMenuDelegate::CreateLambda([](FMenuBuilder& InMenuBuilder) {
+
+			InMenuBuilder.BeginSection("EdGraphSchemaAlignment", LOCTEXT("AlignHeader", "Align"));
+			InMenuBuilder.AddMenuEntry(FGraphEditorCommands::Get().AlignNodesTop);
+			InMenuBuilder.AddMenuEntry(FGraphEditorCommands::Get().AlignNodesMiddle);
+			InMenuBuilder.AddMenuEntry(FGraphEditorCommands::Get().AlignNodesBottom);
+			InMenuBuilder.AddMenuEntry(FGraphEditorCommands::Get().AlignNodesLeft);
+			InMenuBuilder.AddMenuEntry(FGraphEditorCommands::Get().AlignNodesCenter);
+			InMenuBuilder.AddMenuEntry(FGraphEditorCommands::Get().AlignNodesRight);
+			InMenuBuilder.AddMenuEntry(FGraphEditorCommands::Get().StraightenConnections);
+			InMenuBuilder.EndSection();
+
+			InMenuBuilder.BeginSection("EdGraphSchemaDistribution", LOCTEXT("DistributionHeader", "Distribution"));
+			InMenuBuilder.AddMenuEntry(FGraphEditorCommands::Get().DistributeNodesHorizontally);
+			InMenuBuilder.AddMenuEntry(FGraphEditorCommands::Get().DistributeNodesVertically);
+			InMenuBuilder.EndSection();
+		}));
+
 		Context.MenuBuilder->BeginSection("SoundCueGraphNodeEdit");
 		{
 			Context.MenuBuilder->AddMenuEntry( FGenericCommands::Get().Delete );
@@ -250,6 +268,7 @@ void USoundCueGraphNode::GetContextMenuActions(const FGraphNodeContextMenuBuilde
 				Context.MenuBuilder->AddMenuEntry(FSoundCueGraphEditorCommands::Get().BrowserSync);
 			}
 		}
+
 		Context.MenuBuilder->EndSection();
 	}
 }
