Fix issue with re-importing multichannel (3+) wave file not using collection of mono standalone files. Fixes confusion of calculating duration when providing multichannel files and re-import requests mono collection.

diff --git a/Engine/Source/Editor/AudioEditor/Private/Factories/ReimportSoundFactory.cpp b/Engine/Source/Editor/AudioEditor/Private/Factories/ReimportSoundFactory.cpp
--- a/Engine/Source/Editor/AudioEditor/Private/Factories/ReimportSoundFactory.cpp
+++ b/Engine/Source/Editor/AudioEditor/Private/Factories/ReimportSoundFactory.cpp
@@ -33,8 +33,7 @@ UReimportSoundFactory::UReimportSoundFactory(const FObjectInitializer& ObjectIni
 
 bool UReimportSoundFactory::CanReimport(UObject* Obj, TArray<FString>& OutFilenames)
 {
-	USoundWave* SoundWave = Cast<USoundWave>(Obj);
-	if (SoundWave && SoundWave->NumChannels < 3)
+	if (USoundWave* SoundWave = Cast<USoundWave>(Obj))
 	{
 		SoundWave->AssetImportData->ExtractFilenames(OutFilenames);
 
diff --git a/Engine/Source/Editor/AudioEditor/Private/Factories/ReimportSoundSurroundFactory.cpp b/Engine/Source/Editor/AudioEditor/Private/Factories/ReimportSoundSurroundFactory.cpp
--- a/Engine/Source/Editor/AudioEditor/Private/Factories/ReimportSoundSurroundFactory.cpp
+++ b/Engine/Source/Editor/AudioEditor/Private/Factories/ReimportSoundSurroundFactory.cpp
@@ -26,7 +26,11 @@ bool UReimportSoundSurroundFactory::CanReimport(UObject* Obj, TArray<FString>& O
 	if (SoundWave && SoundWave->NumChannels > 2)
 	{
 		FString SourceFilename = SoundWave->AssetImportData->GetFirstFilename();
-		if (!SourceFilename.IsEmpty() && FactoryCanImport(SourceFilename))
+		if (!FactoryCanImport(SourceFilename))
+		{
+			return false;
+		}
+		if (!SourceFilename.IsEmpty())
 		{
 			// Get filename with speaker location removed
 			FString BaseFilename = FPaths::GetBaseFilename(SourceFilename).LeftChop(3);
