Sound thumbnails optimized, avoids loading asset until previewed

diff --git a/Engine/Source/Editor/AudioEditor/Private/AssetTypeActions/AssetTypeActions_SoundBase.cpp b/Engine/Source/Editor/AudioEditor/Private/AssetTypeActions/AssetTypeActions_SoundBase.cpp
--- a/Engine/Source/Editor/AudioEditor/Private/AssetTypeActions/AssetTypeActions_SoundBase.cpp
+++ b/Engine/Source/Editor/AudioEditor/Private/AssetTypeActions/AssetTypeActions_SoundBase.cpp
@@ -106,6 +106,23 @@ bool FAssetTypeActions_SoundBase::IsSoundPlaying(USoundBase* Sound) const
 	return PreviewComp && PreviewComp->Sound == Sound && PreviewComp->IsPlaying();
 }
 
+bool FAssetTypeActions_SoundBase::IsSoundPlaying(const FAssetData& AssetData) const
+{
+	const UAudioComponent* PreviewComp = GEditor->GetPreviewAudioComponent();
+	if (PreviewComp && PreviewComp->Sound && PreviewComp->IsPlaying())
+	{
+		if (PreviewComp->Sound->GetFName() == AssetData.AssetName)
+		{
+			if (PreviewComp->Sound->GetOutermost()->GetFName() == AssetData.PackageName)
+			{
+				return true;
+			}
+		}
+	}
+
+	return false;
+}
+
 void FAssetTypeActions_SoundBase::ExecutePlaySound(TArray<TWeakObjectPtr<USoundBase>> Objects) const
 {
 	for (auto ObjIt = Objects.CreateConstIterator(); ObjIt; ++ObjIt)
@@ -129,14 +146,9 @@ TSharedPtr<SWidget> FAssetTypeActions_SoundBase::GetThumbnailOverlay(const FAsse
 {
 	auto OnGetDisplayBrushLambda = [this, AssetData]() -> const FSlateBrush*
 	{
-		UObject* Asset = AssetData.GetAsset();
-		if (Asset)
+		if (IsSoundPlaying(AssetData))
 		{
-			USoundBase* Sound = CastChecked<USoundBase>(Asset);
-			if (IsSoundPlaying(Sound))
-			{
-				return FEditorStyle::GetBrush("MediaAsset.AssetActions.Stop.Large");
-			}
+			return FEditorStyle::GetBrush("MediaAsset.AssetActions.Stop.Large");
 		}
 
 		return FEditorStyle::GetBrush("MediaAsset.AssetActions.Play.Large");
@@ -144,33 +156,23 @@ TSharedPtr<SWidget> FAssetTypeActions_SoundBase::GetThumbnailOverlay(const FAsse
 
 	auto OnClickedLambda = [this, AssetData]() -> FReply
 	{
-		UObject* Asset = AssetData.GetAsset();
-		if (Asset)
+		if (IsSoundPlaying(AssetData))
 		{
-			USoundBase* Sound = CastChecked<USoundBase>(Asset);
-
-			if (IsSoundPlaying(Sound))
-			{
-				StopSound();
-			}
-			else
-			{
-				PlaySound(Sound);
-			}
+			StopSound();
+		}
+		else
+		{
+			// Load and play sound
+			PlaySound(Cast<USoundBase>(AssetData.GetAsset()));
 		}
 		return FReply::Handled();
 	};
 
 	auto OnToolTipTextLambda = [this, AssetData]() -> FText
 	{
-		UObject* Asset = AssetData.GetAsset();
-		if (Asset)
+		if (IsSoundPlaying(AssetData))
 		{
-			USoundBase* Sound = CastChecked<USoundBase>(Asset);
-			if (IsSoundPlaying(Sound))
-			{
-				return LOCTEXT("Thumbnail_StopSoundToolTip", "Stop selected sound");
-			}
+			return LOCTEXT("Thumbnail_StopSoundToolTip", "Stop selected sound");
 		}
 
 		return LOCTEXT("Thumbnail_PlaySoundToolTip", "Play selected sound");
@@ -184,17 +186,9 @@ TSharedPtr<SWidget> FAssetTypeActions_SoundBase::GetThumbnailOverlay(const FAsse
 
 	auto OnGetVisibilityLambda = [this, Box, AssetData]() -> EVisibility
 	{
-		UObject* Asset = AssetData.GetAsset();
-		if (Asset)
+		if (Box.IsValid() && (Box->IsHovered() || IsSoundPlaying(AssetData)))
 		{
-			USoundBase* Sound = CastChecked<USoundBase>(Asset);
-			if (Box.IsValid())
-			{
-				if (Box->IsHovered() || IsSoundPlaying(Sound))
-				{
-					return EVisibility::Visible;
-				}
-			}
+			return EVisibility::Visible;
 		}
 
 		return EVisibility::Hidden;
diff --git a/Engine/Source/Editor/AudioEditor/Private/AssetTypeActions/AssetTypeActions_SoundWave.cpp b/Engine/Source/Editor/AudioEditor/Private/AssetTypeActions/AssetTypeActions_SoundWave.cpp
--- a/Engine/Source/Editor/AudioEditor/Private/AssetTypeActions/AssetTypeActions_SoundWave.cpp
+++ b/Engine/Source/Editor/AudioEditor/Private/AssetTypeActions/AssetTypeActions_SoundWave.cpp
@@ -217,18 +217,9 @@ TSharedPtr<SWidget> FAssetTypeActions_SoundWave::GetThumbnailOverlay(const FAsse
 {
 	auto OnGetDisplayBrushLambda = [this, AssetData]() -> const FSlateBrush*
 	{
-		if(AssetData.IsAssetLoaded())
+		if (IsSoundPlaying(AssetData))
 		{
-			UObject* Asset = AssetData.GetAsset();
-			if (Asset)
-			{
-				USoundBase* Sound = CastChecked<USoundBase>(AssetData.GetAsset());
-
-				if (IsSoundPlaying(Sound))
-				{
-					return FEditorStyle::GetBrush("MediaAsset.AssetActions.Stop.Large");
-				}
-			}
+			return FEditorStyle::GetBrush("MediaAsset.AssetActions.Stop.Large");
 		}
 
 		return FEditorStyle::GetBrush("MediaAsset.AssetActions.Play.Large");
@@ -236,27 +227,23 @@ TSharedPtr<SWidget> FAssetTypeActions_SoundWave::GetThumbnailOverlay(const FAsse
 
 	auto OnClickedLambda = [this, AssetData]() -> FReply
 	{
-		USoundBase* Sound = CastChecked<USoundBase>(AssetData.GetAsset());
-		if (IsSoundPlaying(Sound))
+		if (IsSoundPlaying(AssetData))
 		{
 			StopSound();
 		}
 		else
 		{
-			PlaySound(Sound);
+			// Load and play sound
+			PlaySound(Cast<USoundBase>(AssetData.GetAsset()));
 		}
 		return FReply::Handled();
 	};
 
 	auto OnToolTipTextLambda = [this, AssetData]() -> FText
 	{
-		if(AssetData.IsAssetLoaded())
+		if (IsSoundPlaying(AssetData))
 		{
-			USoundBase* Sound = CastChecked<USoundBase>(AssetData.GetAsset());
-			if (IsSoundPlaying(Sound))
-			{
-				return LOCTEXT("Blueprint_StopSoundToolTip", "Stop selected Sound Wave");
-			}
+			return LOCTEXT("Blueprint_StopSoundToolTip", "Stop selected Sound Wave");
 		}
 
 		return LOCTEXT("Blueprint_PlaySoundToolTip", "Play selected Sound Wave");
@@ -270,22 +257,7 @@ TSharedPtr<SWidget> FAssetTypeActions_SoundWave::GetThumbnailOverlay(const FAsse
 
 	auto OnGetVisibilityLambda = [this, Box, AssetData]() -> EVisibility
 	{
-		if(AssetData.IsAssetLoaded())
-		{
-			UObject* Asset = AssetData.GetAsset();
-			if (Asset)
-			{
-				USoundBase* Sound = CastChecked<USoundBase>(Asset);
-				if (Box.IsValid())
-				{
-					if (Box->IsHovered() || IsSoundPlaying(Sound))
-					{
-						return EVisibility::Visible;
-					}
-				}
-			}
-		}
-		else if (Box.IsValid() && Box->IsHovered())
+		if (Box.IsValid() && (Box->IsHovered() || IsSoundPlaying(AssetData)))
 		{
 			return EVisibility::Visible;
 		}
diff --git a/Engine/Source/Editor/AudioEditor/Public/AssetTypeActions/AssetTypeActions_SoundBase.h b/Engine/Source/Editor/AudioEditor/Public/AssetTypeActions/AssetTypeActions_SoundBase.h
--- a/Engine/Source/Editor/AudioEditor/Public/AssetTypeActions/AssetTypeActions_SoundBase.h
+++ b/Engine/Source/Editor/AudioEditor/Public/AssetTypeActions/AssetTypeActions_SoundBase.h
@@ -33,6 +33,9 @@ protected:
 	/** Return true if the specified sound is playing */
 	bool IsSoundPlaying(USoundBase* Sound) const;
 
+	/** Return true if the specified asset's sound is playing */
+	bool IsSoundPlaying(const FAssetData& AssetData) const;
+
 private:
 	/** Handler for when PlaySound is selected */
 	void ExecutePlaySound(TArray<TWeakObjectPtr<USoundBase>> Objects) const;
