Fixed up new use of deprecated Import Delegates.

diff --git a/Engine/Source/Editor/AudioEditor/Private/Factories/SoundFactory.cpp b/Engine/Source/Editor/AudioEditor/Private/Factories/SoundFactory.cpp
--- a/Engine/Source/Editor/AudioEditor/Private/Factories/SoundFactory.cpp
+++ b/Engine/Source/Editor/AudioEditor/Private/Factories/SoundFactory.cpp
@@ -164,7 +164,7 @@ UObject* USoundFactory::FactoryCreateBinary
 	{
 		// Unrecognized sound format
 		Warn->Logf(ELogVerbosity::Error, TEXT("Unrecognized sound format '%s' in %s"), FileType, *Name.ToString());
-		FEditorDelegates::OnAssetPostImport.Broadcast(this, nullptr);
+		GEditor->GetEditorSubsystem<UImportSubsystem>()->BroadcastAssetPostImport(this, nullptr);
 	}
 
 	return SoundObject;
