Fix crash in cases where active AudioDevice is null and re-importing while PIE is running.

diff --git a/Engine/Source/Editor/AudioEditor/Private/Factories/SoundFactory.cpp b/Engine/Source/Editor/AudioEditor/Private/Factories/SoundFactory.cpp
--- a/Engine/Source/Editor/AudioEditor/Private/Factories/SoundFactory.cpp
+++ b/Engine/Source/Editor/AudioEditor/Private/Factories/SoundFactory.cpp
@@ -232,11 +232,11 @@ UObject* USoundFactory::CreateObject
 			// to be auditioned in the editor properly.
 			if (!ExistingSound->ResourceData)
 			{
-				FAudioDevice* AudioDevice = AudioDeviceManager->GetActiveAudioDevice();
-				check(AudioDevice);
-
-				FName RuntimeFormat = AudioDevice->GetRuntimeFormat(ExistingSound);
-				ExistingSound->InitAudioResource(RuntimeFormat);
+				if (FAudioDevice* AudioDevice = GEngine->GetMainAudioDevice())
+				{
+					FName RuntimeFormat = AudioDevice->GetRuntimeFormat(ExistingSound);
+					ExistingSound->InitAudioResource(RuntimeFormat);
+				}
 			}
 
 			UE_LOG(LogAudioEditor, Log, TEXT("Stopping Sound Resources of Existing Sound"));
