Fix for mouse decorator and drag-and-drop on sound cue graph

diff --git a/Engine/Source/Editor/AudioEditor/Classes/SoundCueGraph/SoundCueGraphSchema.h b/Engine/Source/Editor/AudioEditor/Classes/SoundCueGraph/SoundCueGraphSchema.h
--- a/Engine/Source/Editor/AudioEditor/Classes/SoundCueGraph/SoundCueGraphSchema.h
+++ b/Engine/Source/Editor/AudioEditor/Classes/SoundCueGraph/SoundCueGraphSchema.h
@@ -122,6 +122,7 @@ class USoundCueGraphSchema : public UEdGraphSchema
 	virtual FLinearColor GetPinTypeColor(const FEdGraphPinType& PinType) const override;
 	virtual void BreakNodeLinks(UEdGraphNode& TargetNode) const override;
 	virtual void BreakPinLinks(UEdGraphPin& TargetPin, bool bSendsNodeNotifcation) const override;
+	virtual void GetAssetsGraphHoverMessage(const TArray<FAssetData>& Assets, const UEdGraph* HoverGraph, FString& OutTooltipText, bool& OutOkIcon) const override;
 	virtual void DroppedAssetsOnGraph(const TArray<struct FAssetData>& Assets, const FVector2D& GraphPosition, UEdGraph* Graph) const override;
 	virtual void DroppedAssetsOnNode(const TArray<FAssetData>& Assets, const FVector2D& GraphPosition, UEdGraphNode* Node) const override;
 	virtual int32 GetNodeSelectionCount(const UEdGraph* Graph) const override;
diff --git a/Engine/Source/Editor/AudioEditor/Private/SoundCueGraphSchema.cpp b/Engine/Source/Editor/AudioEditor/Private/SoundCueGraphSchema.cpp
--- a/Engine/Source/Editor/AudioEditor/Private/SoundCueGraphSchema.cpp
+++ b/Engine/Source/Editor/AudioEditor/Private/SoundCueGraphSchema.cpp
@@ -395,6 +395,22 @@ void USoundCueGraphSchema::BreakPinLinks(UEdGraphPin& TargetPin, bool bSendsNode
 	}
 }
 
+void USoundCueGraphSchema::GetAssetsGraphHoverMessage(const TArray<FAssetData>& Assets, const UEdGraph* HoverGraph, FString& OutTooltipText, bool& OutOkIcon) const
+{
+	OutOkIcon = false;
+
+	for (int32 AssetIdx = 0; AssetIdx < Assets.Num(); ++AssetIdx)
+	{
+		// As soon as one of the items is a sound wave, say we can drag it on... we actually eat only the sound waves.
+		USoundWave* SoundWav = Cast<USoundWave>(Assets[AssetIdx].GetAsset());
+		if (SoundWav)
+		{
+			OutOkIcon = true;
+			break;
+		}
+	}
+}
+
 void USoundCueGraphSchema::DroppedAssetsOnGraph(const TArray<FAssetData>& Assets, const FVector2D& GraphPosition, UEdGraph* Graph) const
 {
 	//////////////////////////////////////////////////////////////////////////
