Clean up asset picker by organizing, standardizing descriptive names, and adding subcategorizing in menu

diff --git a/Engine/Plugins/Runtime/AudioModulation/Source/AudioModulationEditor/Private/AssetTypeActions/AssetTypeActions_SoundControlBus.cpp b/Engine/Plugins/Runtime/AudioModulation/Source/AudioModulationEditor/Private/AssetTypeActions/AssetTypeActions_SoundControlBus.cpp
--- a/Engine/Plugins/Runtime/AudioModulation/Source/AudioModulationEditor/Private/AssetTypeActions/AssetTypeActions_SoundControlBus.cpp
+++ b/Engine/Plugins/Runtime/AudioModulation/Source/AudioModulationEditor/Private/AssetTypeActions/AssetTypeActions_SoundControlBus.cpp
@@ -1,25 +1,56 @@
 // Copyright 1998-2019 Epic Games, Inc. All Rights Reserved.
-
 #include "AssetTypeActions_SoundControlBus.h"
+
 #include "SoundControlBus.h"
 
 
+#define LOCTEXT_NAMESPACE "AssetTypeActions"
+
+namespace Modulation
+{
+	static const TArray<FText> SubMenus
+	{
+		FText(LOCTEXT("AssetSoundMixSubMenu", "Mix"))
+	};
+} // namespace <>
+
 UClass* FAssetTypeActions_SoundVolumeControlBus::GetSupportedClass() const
 {
 	return USoundVolumeControlBus::StaticClass();
 }
 
+const TArray<FText>& FAssetTypeActions_SoundVolumeControlBus::GetSubMenus() const
+{
+	return Modulation::SubMenus;
+}
+
 UClass* FAssetTypeActions_SoundPitchControlBus::GetSupportedClass() const
 {
 	return USoundPitchControlBus::StaticClass();
 }
 
+const TArray<FText>& FAssetTypeActions_SoundPitchControlBus::GetSubMenus() const
+{
+	return Modulation::SubMenus;
+}
+
 UClass* FAssetTypeActions_SoundLPFControlBus::GetSupportedClass() const
 {
 	return USoundLPFControlBus::StaticClass();
 }
 
+const TArray<FText>& FAssetTypeActions_SoundLPFControlBus::GetSubMenus() const
+{
+	return Modulation::SubMenus;
+}
+
 UClass* FAssetTypeActions_SoundHPFControlBus::GetSupportedClass() const
 {
 	return USoundHPFControlBus::StaticClass();
 }
+
+const TArray<FText>& FAssetTypeActions_SoundHPFControlBus::GetSubMenus() const
+{
+	return Modulation::SubMenus;
+}
+#undef LOCTEXT_NAMESPACE
diff --git a/Engine/Plugins/Runtime/AudioModulation/Source/AudioModulationEditor/Private/AssetTypeActions/AssetTypeActions_SoundControlBus.h b/Engine/Plugins/Runtime/AudioModulation/Source/AudioModulationEditor/Private/AssetTypeActions/AssetTypeActions_SoundControlBus.h
--- a/Engine/Plugins/Runtime/AudioModulation/Source/AudioModulationEditor/Private/AssetTypeActions/AssetTypeActions_SoundControlBus.h
+++ b/Engine/Plugins/Runtime/AudioModulation/Source/AudioModulationEditor/Private/AssetTypeActions/AssetTypeActions_SoundControlBus.h
@@ -1,10 +1,10 @@
 // Copyright 1998-2019 Epic Games, Inc. All Rights Reserved.
-
 #pragma once
 
 #include "CoreMinimal.h"
 #include "AssetTypeActions_Base.h"
 
+
 class FAssetTypeActions_SoundVolumeControlBus : public FAssetTypeActions_Base
 {
 public:
@@ -13,6 +13,8 @@ public:
 	virtual FColor GetTypeColor() const override { return FColor(33, 183, 0); }
 	virtual UClass* GetSupportedClass() const override;
 	virtual uint32 GetCategories() override { return EAssetTypeCategories::Sounds; }
+
+	virtual const TArray<FText>& GetSubMenus() const override;
 };
 
 class FAssetTypeActions_SoundPitchControlBus : public FAssetTypeActions_Base
@@ -23,6 +25,8 @@ public:
 	virtual FColor GetTypeColor() const override { return FColor(181, 21, 0); }
 	virtual UClass* GetSupportedClass() const override;
 	virtual uint32 GetCategories() override { return EAssetTypeCategories::Sounds; }
+
+	virtual const TArray<FText>& GetSubMenus() const override;
 };
 
 class FAssetTypeActions_SoundLPFControlBus : public FAssetTypeActions_Base
@@ -33,6 +37,8 @@ public:
 	virtual FColor GetTypeColor() const override { return FColor(0, 156, 183); }
 	virtual UClass* GetSupportedClass() const override;
 	virtual uint32 GetCategories() override { return EAssetTypeCategories::Sounds; }
+
+	virtual const TArray<FText>& GetSubMenus() const override;
 };
 
 class FAssetTypeActions_SoundHPFControlBus : public FAssetTypeActions_Base
@@ -43,4 +49,6 @@ public:
 	virtual FColor GetTypeColor() const override { return FColor(94, 237, 183); }
 	virtual UClass* GetSupportedClass() const override;
 	virtual uint32 GetCategories() override { return EAssetTypeCategories::Sounds; }
+
+	virtual const TArray<FText>& GetSubMenus() const override;
 };
diff --git a/Engine/Plugins/Runtime/AudioModulation/Source/AudioModulationEditor/Private/AssetTypeActions/AssetTypeActions_SoundControlBusMix.cpp b/Engine/Plugins/Runtime/AudioModulation/Source/AudioModulationEditor/Private/AssetTypeActions/AssetTypeActions_SoundControlBusMix.cpp
--- a/Engine/Plugins/Runtime/AudioModulation/Source/AudioModulationEditor/Private/AssetTypeActions/AssetTypeActions_SoundControlBusMix.cpp
+++ b/Engine/Plugins/Runtime/AudioModulation/Source/AudioModulationEditor/Private/AssetTypeActions/AssetTypeActions_SoundControlBusMix.cpp
@@ -1,10 +1,23 @@
 // Copyright 1998-2019 Epic Games, Inc. All Rights Reserved.
-
 #include "AssetTypeActions_SoundControlBusMix.h"
+
 #include "SoundControlBusMix.h"
 
+#define LOCTEXT_NAMESPACE "AssetTypeActions"
 
 UClass* FAssetTypeActions_SoundControlBusMix::GetSupportedClass() const
 {
 	return USoundControlBusMix::StaticClass();
 }
+
+const TArray<FText>& FAssetTypeActions_SoundControlBusMix::GetSubMenus() const
+{
+	static const TArray<FText> SubMenus
+	{
+		FText(LOCTEXT("AssetSoundMixSubMenu", "Mix"))
+	};
+
+	return SubMenus;
+}
+
+#undef LOCTEXT_NAMESPACE
diff --git a/Engine/Plugins/Runtime/AudioModulation/Source/AudioModulationEditor/Private/AssetTypeActions/AssetTypeActions_SoundControlBusMix.h b/Engine/Plugins/Runtime/AudioModulation/Source/AudioModulationEditor/Private/AssetTypeActions/AssetTypeActions_SoundControlBusMix.h
--- a/Engine/Plugins/Runtime/AudioModulation/Source/AudioModulationEditor/Private/AssetTypeActions/AssetTypeActions_SoundControlBusMix.h
+++ b/Engine/Plugins/Runtime/AudioModulation/Source/AudioModulationEditor/Private/AssetTypeActions/AssetTypeActions_SoundControlBusMix.h
@@ -9,8 +9,9 @@ class FAssetTypeActions_SoundControlBusMix : public FAssetTypeActions_Base
 {
 public:
 	// IAssetTypeActions Implementation
-	virtual FText GetName() const override { return NSLOCTEXT("AssetTypeActions", "AssetTypeActions_SoundControlBusMix", "Modulation Mix"); }
+	virtual FText GetName() const override { return NSLOCTEXT("AssetTypeActions", "AssetTypeActions_SoundControlBusMix", "Control Bus Mix"); }
 	virtual FColor GetTypeColor() const override { return FColor(237, 237, 0); }
+	virtual const TArray<FText>& GetSubMenus() const override;
 	virtual UClass* GetSupportedClass() const override;
 	virtual uint32 GetCategories() override { return EAssetTypeCategories::Sounds; }
 };
diff --git a/Engine/Plugins/Runtime/AudioModulation/Source/AudioModulationEditor/Private/AssetTypeActions/AssetTypeActions_SoundModulationSettings.cpp b/Engine/Plugins/Runtime/AudioModulation/Source/AudioModulationEditor/Private/AssetTypeActions/AssetTypeActions_SoundModulationSettings.cpp
--- a/Engine/Plugins/Runtime/AudioModulation/Source/AudioModulationEditor/Private/AssetTypeActions/AssetTypeActions_SoundModulationSettings.cpp
+++ b/Engine/Plugins/Runtime/AudioModulation/Source/AudioModulationEditor/Private/AssetTypeActions/AssetTypeActions_SoundModulationSettings.cpp
@@ -4,7 +4,20 @@
 #include "SoundModulationPatch.h"
 
 
+#define LOCTEXT_NAMESPACE "AssetTypeActions"
+
 UClass* FAssetTypeActions_SoundModulationSettings::GetSupportedClass() const
 {
 	return USoundModulationSettings::StaticClass();
 }
+
+const TArray<FText>& FAssetTypeActions_SoundModulationSettings::GetSubMenus() const
+{
+	static const TArray<FText> SubMenus
+	{
+		FText(LOCTEXT("AssetSoundMixSubMenu", "Mix"))
+	};
+
+	return SubMenus;
+}
+#undef LOCTEXT_NAMESPACE
diff --git a/Engine/Plugins/Runtime/AudioModulation/Source/AudioModulationEditor/Private/AssetTypeActions/AssetTypeActions_SoundModulationSettings.h b/Engine/Plugins/Runtime/AudioModulation/Source/AudioModulationEditor/Private/AssetTypeActions/AssetTypeActions_SoundModulationSettings.h
--- a/Engine/Plugins/Runtime/AudioModulation/Source/AudioModulationEditor/Private/AssetTypeActions/AssetTypeActions_SoundModulationSettings.h
+++ b/Engine/Plugins/Runtime/AudioModulation/Source/AudioModulationEditor/Private/AssetTypeActions/AssetTypeActions_SoundModulationSettings.h
@@ -11,6 +11,7 @@ public:
 	// IAssetTypeActions Implementation
 	virtual FText GetName() const override { return NSLOCTEXT("AssetTypeActions", "AssetTypeActions_SoundModulationSettings", "Modulation Settings"); }
 	virtual FColor GetTypeColor() const override { return FColor(237, 237, 0); }
+	virtual const TArray<FText>& GetSubMenus() const override;
 	virtual UClass* GetSupportedClass() const override;
 	virtual uint32 GetCategories() override { return EAssetTypeCategories::Sounds; }
 };
diff --git a/Engine/Plugins/Runtime/AudioModulation/Source/AudioModulationEditor/Private/AssetTypeActions/AssetTypeActions_SoundModulatorLFO.cpp b/Engine/Plugins/Runtime/AudioModulation/Source/AudioModulationEditor/Private/AssetTypeActions/AssetTypeActions_SoundModulatorLFO.cpp
--- a/Engine/Plugins/Runtime/AudioModulation/Source/AudioModulationEditor/Private/AssetTypeActions/AssetTypeActions_SoundModulatorLFO.cpp
+++ b/Engine/Plugins/Runtime/AudioModulation/Source/AudioModulationEditor/Private/AssetTypeActions/AssetTypeActions_SoundModulatorLFO.cpp
@@ -1,10 +1,23 @@
 // Copyright 1998-2019 Epic Games, Inc. All Rights Reserved.
-
 #include "AssetTypeActions_SoundModulatorLFO.h"
+
 #include "SoundModulatorLFO.h"
 
 
+#define LOCTEXT_NAMESPACE "AssetTypeActions"
+
 UClass* FAssetTypeActions_SoundModulatorLFO::GetSupportedClass() const
 {
 	return USoundBusModulatorLFO::StaticClass();
 }
+
+const TArray<FText>& FAssetTypeActions_SoundModulatorLFO::GetSubMenus() const
+{
+	static const TArray<FText> SubMenus
+	{
+		FText(LOCTEXT("AssetSoundMixSubMenu", "Mix"))
+	};
+
+	return SubMenus;
+}
+#undef LOCTEXT_NAMESPACE
\ No newline at end of file
diff --git a/Engine/Plugins/Runtime/AudioModulation/Source/AudioModulationEditor/Private/AssetTypeActions/AssetTypeActions_SoundModulatorLFO.h b/Engine/Plugins/Runtime/AudioModulation/Source/AudioModulationEditor/Private/AssetTypeActions/AssetTypeActions_SoundModulatorLFO.h
--- a/Engine/Plugins/Runtime/AudioModulation/Source/AudioModulationEditor/Private/AssetTypeActions/AssetTypeActions_SoundModulatorLFO.h
+++ b/Engine/Plugins/Runtime/AudioModulation/Source/AudioModulationEditor/Private/AssetTypeActions/AssetTypeActions_SoundModulatorLFO.h
@@ -9,8 +9,9 @@ class FAssetTypeActions_SoundModulatorLFO : public FAssetTypeActions_Base
 {
 public:
 	// IAssetTypeActions Implementation
-	virtual FText GetName() const override { return NSLOCTEXT("AssetTypeActions", "AssetTypeActions_SoundModulatorLFO", "Modulator LFO"); }
+	virtual FText GetName() const override { return NSLOCTEXT("AssetTypeActions", "AssetTypeActions_SoundModulatorLFO", "Control Modulator (LFO)"); }
 	virtual FColor GetTypeColor() const override { return FColor(237, 237, 0); }
+	virtual const TArray<FText>& GetSubMenus() const override;
 	virtual UClass* GetSupportedClass() const override;
 	virtual uint32 GetCategories() override { return EAssetTypeCategories::Sounds; }
 };
diff --git a/Engine/Plugins/Runtime/Oculus/OculusAudio/Source/OculusAudioEditor/Private/IOculusAudioEditorPlugin.cpp b/Engine/Plugins/Runtime/Oculus/OculusAudio/Source/OculusAudioEditor/Private/IOculusAudioEditorPlugin.cpp
--- a/Engine/Plugins/Runtime/Oculus/OculusAudio/Source/OculusAudioEditor/Private/IOculusAudioEditorPlugin.cpp
+++ b/Engine/Plugins/Runtime/Oculus/OculusAudio/Source/OculusAudioEditor/Private/IOculusAudioEditorPlugin.cpp
@@ -3,6 +3,7 @@
 #include "IOculusAudioEditorPlugin.h"
 #include "OculusAudioSettings.h"
 #include "OculusAudioSourceSettingsFactory.h"
+#include "OculusAmbisonicSettingsFactory.h"
 
 #include "Modules/ModuleManager.h"
 #include "ISettingsModule.h"
@@ -12,8 +13,8 @@ void FOculusAudioEditorPlugin::StartupModule()
     // Register asset types
     IAssetTools& AssetTools = FModuleManager::LoadModuleChecked<FAssetToolsModule>("AssetTools").Get();
     {
-        TSharedRef<IAssetTypeActions> Action = MakeShareable(new FAssetTypeActions_OculusAudioSourceSettings);
-        AssetTools.RegisterAssetTypeActions(Action);
+		AssetTools.RegisterAssetTypeActions(MakeShareable(new FAssetTypeActions_OculusAudioSourceSettings));
+		AssetTools.RegisterAssetTypeActions(MakeShareable(new FAssetTypeActions_OculusAmbisonicsSettings));
 
         ISettingsModule* SettingsModule = FModuleManager::Get().GetModulePtr<ISettingsModule>("Settings");
         if (SettingsModule)
diff --git a/Engine/Plugins/Runtime/Oculus/OculusAudio/Source/OculusAudioEditor/Private/OculusAmbisonicSettingsFactory.cpp b/Engine/Plugins/Runtime/Oculus/OculusAudio/Source/OculusAudioEditor/Private/OculusAmbisonicSettingsFactory.cpp
--- a/Engine/Plugins/Runtime/Oculus/OculusAudio/Source/OculusAudioEditor/Private/OculusAmbisonicSettingsFactory.cpp
+++ b/Engine/Plugins/Runtime/Oculus/OculusAudio/Source/OculusAudioEditor/Private/OculusAmbisonicSettingsFactory.cpp
@@ -6,9 +6,11 @@
 #include "UObject/ObjectMacros.h"
 #include "UObject/Object.h"
 
+#define LOCTEXT_NAMESPACE "AssetTypeActions"
+
 FText FAssetTypeActions_OculusAmbisonicsSettings::GetName() const
 {
-    return NSLOCTEXT("AssetTypeActions", "AssetTypeActions_OculusAmbisonicsSettings", "Oculus Audio Source Settings");
+    return LOCTEXT("AssetTypeActions_OculusAmbisonicsSettings", "Oculus Audio Source Settings");
 }
 
 FColor FAssetTypeActions_OculusAmbisonicsSettings::GetTypeColor() const
@@ -16,6 +18,16 @@ FColor FAssetTypeActions_OculusAmbisonicsSettings::GetTypeColor() const
     return FColor(100, 100, 100);
 }
 
+const TArray<FText>& FAssetTypeActions_OculusAmbisonicsSettings::GetSubMenus() const
+{
+	static const TArray<FText> SubMenus
+	{
+		LOCTEXT("AssetSoundOculusSubMenu", "Oculus")
+	};
+
+	return SubMenus;
+}
+
 UClass* FAssetTypeActions_OculusAmbisonicsSettings::GetSupportedClass() const
 {
     return UOculusAmbisonicsSettings::StaticClass();
@@ -45,4 +57,5 @@ UObject* UOculusAmbisonicsSettingsFactory::FactoryCreateNew(UClass* Class, UObje
 uint32 UOculusAmbisonicsSettingsFactory::GetMenuCategories() const
 {
     return EAssetTypeCategories::Sounds;
-}
\ No newline at end of file
+}
+#undef LOCTEXT_NAMESPACE
diff --git a/Engine/Plugins/Runtime/Oculus/OculusAudio/Source/OculusAudioEditor/Private/OculusAmbisonicSettingsFactory.h b/Engine/Plugins/Runtime/Oculus/OculusAudio/Source/OculusAudioEditor/Private/OculusAmbisonicSettingsFactory.h
--- a/Engine/Plugins/Runtime/Oculus/OculusAudio/Source/OculusAudioEditor/Private/OculusAmbisonicSettingsFactory.h
+++ b/Engine/Plugins/Runtime/Oculus/OculusAudio/Source/OculusAudioEditor/Private/OculusAmbisonicSettingsFactory.h
@@ -13,6 +13,7 @@ class FAssetTypeActions_OculusAmbisonicsSettings : public FAssetTypeActions_Base
 public:
     virtual FText GetName() const override;
     virtual FColor GetTypeColor() const override;
+	virtual const TArray<FText>& GetSubMenus() const override;
     virtual UClass* GetSupportedClass() const override;
     virtual uint32 GetCategories() override;
 };
diff --git a/Engine/Plugins/Runtime/Oculus/OculusAudio/Source/OculusAudioEditor/Private/OculusAudioSourceSettingsFactory.cpp b/Engine/Plugins/Runtime/Oculus/OculusAudio/Source/OculusAudioEditor/Private/OculusAudioSourceSettingsFactory.cpp
--- a/Engine/Plugins/Runtime/Oculus/OculusAudio/Source/OculusAudioEditor/Private/OculusAudioSourceSettingsFactory.cpp
+++ b/Engine/Plugins/Runtime/Oculus/OculusAudio/Source/OculusAudioEditor/Private/OculusAudioSourceSettingsFactory.cpp
@@ -6,9 +6,21 @@
 #include "UObject/ObjectMacros.h"
 #include "UObject/Object.h"
 
+#define LOCTEXT_NAMESPACE "AssetTypeActions"
+
 FText FAssetTypeActions_OculusAudioSourceSettings::GetName() const
 {
-    return NSLOCTEXT("AssetTypeActions", "AssetTypeActions_OculusAudioSpatializationSettings", "Oculus Audio Source Settings");
+    return LOCTEXT("AssetTypeActions_OculusAudioSpatializationSettings", "Oculus Audio Source Settings");
+}
+
+const TArray<FText>& FAssetTypeActions_OculusAudioSourceSettings::GetSubMenus() const
+{
+	static const TArray<FText> SubMenus
+	{
+		LOCTEXT("AssetSoundOculusSubMenu", "Oculus")
+	};
+
+	return SubMenus;
 }
 
 FColor FAssetTypeActions_OculusAudioSourceSettings::GetTypeColor() const
diff --git a/Engine/Plugins/Runtime/Oculus/OculusAudio/Source/OculusAudioEditor/Private/OculusAudioSourceSettingsFactory.h b/Engine/Plugins/Runtime/Oculus/OculusAudio/Source/OculusAudioEditor/Private/OculusAudioSourceSettingsFactory.h
--- a/Engine/Plugins/Runtime/Oculus/OculusAudio/Source/OculusAudioEditor/Private/OculusAudioSourceSettingsFactory.h
+++ b/Engine/Plugins/Runtime/Oculus/OculusAudio/Source/OculusAudioEditor/Private/OculusAudioSourceSettingsFactory.h
@@ -13,7 +13,8 @@ class FAssetTypeActions_OculusAudioSourceSettings : public FAssetTypeActions_Bas
 public:
     virtual FText GetName() const override;
     virtual FColor GetTypeColor() const override;
-    virtual UClass* GetSupportedClass() const override;
+	virtual const TArray<FText>& GetSubMenus() const override;
+	virtual UClass* GetSupportedClass() const override;
     virtual uint32 GetCategories() override;
 };
 
diff --git a/Engine/Plugins/Runtime/Synthesis/Source/SynthesisEditor/Classes/EpicSynth1PresetBank.h b/Engine/Plugins/Runtime/Synthesis/Source/SynthesisEditor/Classes/EpicSynth1PresetBank.h
--- a/Engine/Plugins/Runtime/Synthesis/Source/SynthesisEditor/Classes/EpicSynth1PresetBank.h
+++ b/Engine/Plugins/Runtime/Synthesis/Source/SynthesisEditor/Classes/EpicSynth1PresetBank.h
@@ -14,6 +14,7 @@ class FAssetTypeActions_ModularSynthPresetBank : public FAssetTypeActions_Base
 public:
 	virtual FText GetName() const override { return NSLOCTEXT("AssetTypeActions", "AssetTypeActions_EpicSynth1PresetBank", "Modular Synth Preset Bank"); }
 	virtual FColor GetTypeColor() const override { return FColor(175, 100, 155); }
+	virtual const TArray<FText>& GetSubMenus() const override;
 	virtual UClass* GetSupportedClass() const override;
 	virtual uint32 GetCategories() override { return EAssetTypeCategories::Sounds; }
 };
@@ -27,6 +28,3 @@ class UModularSynthPresetBankFactory : public UFactory
 	virtual UObject* FactoryCreateNew(UClass* Class, UObject* InParent, FName Name, EObjectFlags Flags, UObject* Context, FFeedbackContext* Warn) override;
 	//~ Begin UFactory Interface	
 };
-
-
-
diff --git a/Engine/Plugins/Runtime/Synthesis/Source/SynthesisEditor/Classes/MonoWaveTablePresetBank.h b/Engine/Plugins/Runtime/Synthesis/Source/SynthesisEditor/Classes/MonoWaveTablePresetBank.h
--- a/Engine/Plugins/Runtime/Synthesis/Source/SynthesisEditor/Classes/MonoWaveTablePresetBank.h
+++ b/Engine/Plugins/Runtime/Synthesis/Source/SynthesisEditor/Classes/MonoWaveTablePresetBank.h
@@ -15,6 +15,7 @@ public:
 	virtual FText GetName() const override { return NSLOCTEXT("AssetTypeActions", "AssetTypeActions_MonoWaveTableSynthPreset", "Mono Wave Table Synth Preset Bank"); }
 	virtual FColor GetTypeColor() const override { return FColor(0, 255, 0); }
 	virtual UClass* GetSupportedClass() const override;
+	virtual const TArray<FText>& GetSubMenus() const override;
 	virtual uint32 GetCategories() override { return EAssetTypeCategories::Sounds; }
 };
 
diff --git a/Engine/Plugins/Runtime/Synthesis/Source/SynthesisEditor/Private/EpicSynth1PresetBank.cpp b/Engine/Plugins/Runtime/Synthesis/Source/SynthesisEditor/Private/EpicSynth1PresetBank.cpp
--- a/Engine/Plugins/Runtime/Synthesis/Source/SynthesisEditor/Private/EpicSynth1PresetBank.cpp
+++ b/Engine/Plugins/Runtime/Synthesis/Source/SynthesisEditor/Private/EpicSynth1PresetBank.cpp
@@ -1,15 +1,28 @@
 // Copyright 1998-2019 Epic Games, Inc. All Rights Reserved.
-
 #include "EpicSynth1PresetBank.h"
+
 #include "SynthComponents/EpicSynth1Component.h"
 #include "UObject/ObjectMacros.h"
 #include "UObject/Object.h"
 
+
+#define LOCTEXT_NAMESPACE "AssetTypeActions"
+
 UClass* FAssetTypeActions_ModularSynthPresetBank::GetSupportedClass() const
 {
 	return UModularSynthPresetBank::StaticClass();
 }
 
+const TArray<FText>& FAssetTypeActions_ModularSynthPresetBank::GetSubMenus() const
+{
+	static const TArray<FText> SubMenus
+	{
+		FText(LOCTEXT("AssetSoundSynthesisSubMenu", "Synthesis"))
+	};
+
+	return SubMenus;
+}
+
 UModularSynthPresetBankFactory::UModularSynthPresetBankFactory(const FObjectInitializer& ObjectInitializer)
 	: Super(ObjectInitializer)
 {
@@ -25,4 +38,6 @@ UObject* UModularSynthPresetBankFactory::FactoryCreateNew(UClass* Class, UObject
 	UModularSynthPresetBank* NewPresetBank = NewObject<UModularSynthPresetBank>(InParent, InName, Flags);
 
 	return NewPresetBank;
-}
\ No newline at end of file
+}
+
+#undef LOCTEXT_NAMESPACE
\ No newline at end of file
diff --git a/Engine/Plugins/Runtime/Synthesis/Source/SynthesisEditor/Private/MonoWaveTablePresetBank.cpp b/Engine/Plugins/Runtime/Synthesis/Source/SynthesisEditor/Private/MonoWaveTablePresetBank.cpp
--- a/Engine/Plugins/Runtime/Synthesis/Source/SynthesisEditor/Private/MonoWaveTablePresetBank.cpp
+++ b/Engine/Plugins/Runtime/Synthesis/Source/SynthesisEditor/Private/MonoWaveTablePresetBank.cpp
@@ -1,15 +1,28 @@
 // Copyright 1998-2019 Epic Games, Inc. All Rights Reserved.
-
 #include "MonoWaveTablePresetBank.h"
+
 #include "SynthComponents/SynthComponentMonoWaveTable.h"
 #include "UObject/ObjectMacros.h"
 #include "UObject/Object.h"
 
+#define LOCTEXT_NAMESPACE "AssetTypeActions"
+
+
 UClass* FAssetTypeActions_MonoWaveTableSynthPreset::GetSupportedClass() const
 {
 	return UMonoWaveTableSynthPreset::StaticClass();
 }
 
+const TArray<FText>& FAssetTypeActions_MonoWaveTableSynthPreset::GetSubMenus() const
+{
+	static const TArray<FText> SubMenus
+	{
+		FText(LOCTEXT("AssetSoundSynthesisSubMenu", "Synthesis"))
+	};
+
+	return SubMenus;
+}
+
 UMonoWaveTableSynthPresetFactory::UMonoWaveTableSynthPresetFactory(const FObjectInitializer& ObjectInitializer)
 	: Super(ObjectInitializer)
 {
@@ -25,4 +38,5 @@ UObject* UMonoWaveTableSynthPresetFactory::FactoryCreateNew(UClass* Class, UObje
 	UMonoWaveTableSynthPreset* NewPresetBank = NewObject<UMonoWaveTableSynthPreset>(InParent, InName, Flags);
 
 	return NewPresetBank;
-}
\ No newline at end of file
+}
+#undef LOCTEXT_NAMESPACE
\ No newline at end of file
diff --git a/Engine/Plugins/Runtime/TimeSynth/Source/TimeSynthEditor/Classes/TimeSynthClip.h b/Engine/Plugins/Runtime/TimeSynth/Source/TimeSynthEditor/Classes/TimeSynthClip.h
--- a/Engine/Plugins/Runtime/TimeSynth/Source/TimeSynthEditor/Classes/TimeSynthClip.h
+++ b/Engine/Plugins/Runtime/TimeSynth/Source/TimeSynthEditor/Classes/TimeSynthClip.h
@@ -16,6 +16,7 @@ class FAssetTypeActions_TimeSynthClip : public FAssetTypeActions_Base
 public:
 	virtual FText GetName() const override { return NSLOCTEXT("AssetTypeActions", "AssetTypeActions_TimeSynthClip", "Time Synth Clip"); }
 	virtual FColor GetTypeColor() const override { return FColor(155, 175, 100); }
+	virtual const TArray<FText>& GetSubMenus() const override;
 	virtual UClass* GetSupportedClass() const override;
 	virtual uint32 GetCategories() override { return EAssetTypeCategories::Sounds; }
 };
diff --git a/Engine/Plugins/Runtime/TimeSynth/Source/TimeSynthEditor/Classes/TimeSynthVolumeGroup.h b/Engine/Plugins/Runtime/TimeSynth/Source/TimeSynthEditor/Classes/TimeSynthVolumeGroup.h
--- a/Engine/Plugins/Runtime/TimeSynth/Source/TimeSynthEditor/Classes/TimeSynthVolumeGroup.h
+++ b/Engine/Plugins/Runtime/TimeSynth/Source/TimeSynthEditor/Classes/TimeSynthVolumeGroup.h
@@ -1,19 +1,22 @@
 // Copyright 1998-2019 Epic Games, Inc. All Rights Reserved.
-
 #pragma once
 
 #include "CoreMinimal.h"
-#include "Stats/Stats.h"
+
 #include "AssetToolsModule.h"
-#include "Factories/Factory.h"
 #include "AssetTypeActions_Base.h"
+#include "Factories/Factory.h"
+#include "Stats/Stats.h"
+
 #include "TimeSynthVolumeGroup.generated.h"
 
+
 class FAssetTypeActions_TimeSynthVolumeGroup : public FAssetTypeActions_Base
 {
 public:
 	virtual FText GetName() const override { return NSLOCTEXT("AssetTypeActions", "AssetTypeActions_TimeSynthVolumeGroup", "Time Synth Volume Group"); }
 	virtual FColor GetTypeColor() const override { return FColor(175, 155, 100); }
+	virtual const TArray<FText>& GetSubMenus() const override;
 	virtual UClass* GetSupportedClass() const override;
 	virtual uint32 GetCategories() override { return EAssetTypeCategories::Sounds; }
 };
@@ -25,8 +28,5 @@ class UTimeSynthVolumeGroupFactory : public UFactory
 
 	//~ Begin UFactory Interface
 	virtual UObject* FactoryCreateNew(UClass* Class, UObject* InParent, FName Name, EObjectFlags Flags, UObject* Context, FFeedbackContext* Warn) override;
-	//~ Begin UFactory Interface	
+	//~ Begin UFactory Interface
 };
-
-
-
diff --git a/Engine/Plugins/Runtime/TimeSynth/Source/TimeSynthEditor/Private/TimeSynthClip.cpp b/Engine/Plugins/Runtime/TimeSynth/Source/TimeSynthEditor/Private/TimeSynthClip.cpp
--- a/Engine/Plugins/Runtime/TimeSynth/Source/TimeSynthEditor/Private/TimeSynthClip.cpp
+++ b/Engine/Plugins/Runtime/TimeSynth/Source/TimeSynthEditor/Private/TimeSynthClip.cpp
@@ -1,17 +1,29 @@
 // Copyright 1998-2019 Epic Games, Inc. All Rights Reserved.
-
 #include "TimeSynthClip.h"
+
 #include "TimeSynthComponent.h"
 #include "Sound/SoundWave.h"
 #include "UObject/ObjectMacros.h"
 #include "UObject/Object.h"
 
 
+#define LOCTEXT_NAMESPACE "AssetTypeActions"
+
 UClass* FAssetTypeActions_TimeSynthClip::GetSupportedClass() const
 {
 	return UTimeSynthClip::StaticClass();
 }
 
+const TArray<FText>& FAssetTypeActions_TimeSynthClip::GetSubMenus() const
+{
+	static const TArray<FText> SubMenus
+	{
+		FText(LOCTEXT("AssetSoundSynthesisSubMenu", "Synthesis"))
+	};
+
+	return SubMenus;
+}
+
 UTimeSynthClipFactory::UTimeSynthClipFactory(const FObjectInitializer& ObjectInitializer)
 	: Super(ObjectInitializer)
 {
@@ -37,4 +49,6 @@ UObject* UTimeSynthClipFactory::FactoryCreateNew(UClass* Class, UObject* InParen
 	}
 	
 	return NewTimeSynthClip;
-}
\ No newline at end of file
+}
+
+#undef LOCTEXT_NAMESPACE
\ No newline at end of file
diff --git a/Engine/Plugins/Runtime/TimeSynth/Source/TimeSynthEditor/Private/TimeSynthVolumeGroup.cpp b/Engine/Plugins/Runtime/TimeSynth/Source/TimeSynthEditor/Private/TimeSynthVolumeGroup.cpp
--- a/Engine/Plugins/Runtime/TimeSynth/Source/TimeSynthEditor/Private/TimeSynthVolumeGroup.cpp
+++ b/Engine/Plugins/Runtime/TimeSynth/Source/TimeSynthEditor/Private/TimeSynthVolumeGroup.cpp
@@ -1,15 +1,28 @@
 // Copyright 1998-2019 Epic Games, Inc. All Rights Reserved.
-
 #include "TimeSynthVolumeGroup.h"
+
 #include "TimeSynthComponent.h"
 #include "UObject/ObjectMacros.h"
 #include "UObject/Object.h"
 
+
+#define LOCTEXT_NAMESPACE "AssetTypeActions"
+
 UClass* FAssetTypeActions_TimeSynthVolumeGroup::GetSupportedClass() const
 {
 	return UTimeSynthVolumeGroup::StaticClass();
 }
 
+const TArray<FText>& FAssetTypeActions_TimeSynthVolumeGroup::GetSubMenus() const
+{
+	static const TArray<FText> SubMenus
+	{
+		FText(LOCTEXT("AssetSoundSynthesisSubMenu", "Synthesis"))
+	};
+
+	return SubMenus;
+}
+
 UTimeSynthVolumeGroupFactory::UTimeSynthVolumeGroupFactory(const FObjectInitializer& ObjectInitializer)
 	: Super(ObjectInitializer)
 {
@@ -25,4 +38,6 @@ UObject* UTimeSynthVolumeGroupFactory::FactoryCreateNew(UClass* Class, UObject*
 	UTimeSynthVolumeGroup* NewTimeSynthVolumeGroup = NewObject<UTimeSynthVolumeGroup>(InParent, InName, Flags);
 
 	return NewTimeSynthVolumeGroup;
-}
\ No newline at end of file
+}
+
+#undef LOCTEXT_NAMESPACE
\ No newline at end of file
diff --git a/Engine/Source/Editor/AudioEditor/Private/AssetTypeActions/AssetTypeActions_DialogueVoice.h b/Engine/Source/Editor/AudioEditor/Private/AssetTypeActions/AssetTypeActions_DialogueVoice.h
--- a/Engine/Source/Editor/AudioEditor/Private/AssetTypeActions/AssetTypeActions_DialogueVoice.h
+++ b/Engine/Source/Editor/AudioEditor/Private/AssetTypeActions/AssetTypeActions_DialogueVoice.h
@@ -6,6 +6,8 @@
 #include "AssetTypeActions_Base.h"
 #include "Sound/DialogueVoice.h"
 
+#define LOCTEXT_NAMESPACE "Dialogue"
+
 class FAssetTypeActions_DialogueVoice : public FAssetTypeActions_Base
 {
 public:
@@ -14,5 +16,16 @@ public:
 	virtual FColor GetTypeColor() const override { return FColor(97, 85, 212); }
 	virtual UClass* GetSupportedClass() const override { return UDialogueVoice::StaticClass(); }
 	virtual uint32 GetCategories() override { return EAssetTypeCategories::Sounds; }
+	virtual const TArray<FText>& GetSubMenus() const override
+	{
+		static const TArray<FText> SubMenus
+		{
+			FText(LOCTEXT("AssetDialogueSubMenu", "Dialogue"))
+		};
+
+		return SubMenus;
+	}
 	virtual bool CanFilter() override { return true; }
 };
+
+#undef LOCTEXT_NAMESPACE
\ No newline at end of file
diff --git a/Engine/Source/Editor/AudioEditor/Private/AssetTypeActions/AssetTypeActions_DialogueWave.cpp b/Engine/Source/Editor/AudioEditor/Private/AssetTypeActions/AssetTypeActions_DialogueWave.cpp
--- a/Engine/Source/Editor/AudioEditor/Private/AssetTypeActions/AssetTypeActions_DialogueWave.cpp
+++ b/Engine/Source/Editor/AudioEditor/Private/AssetTypeActions/AssetTypeActions_DialogueWave.cpp
@@ -195,4 +195,13 @@ void FAssetTypeActions_DialogueWave::ExecuteCreateSoundCue(TArray<TWeakObjectPtr
 	}
 }
 
+const TArray<FText>& FAssetTypeActions_DialogueWave::GetSubMenus() const
+{
+	static const TArray<FText> SubMenus
+	{
+		FText(LOCTEXT("AssetDialogueSubMenu", "Dialogue"))
+	};
+
+	return SubMenus;
+}
 #undef LOCTEXT_NAMESPACE
diff --git a/Engine/Source/Editor/AudioEditor/Private/AssetTypeActions/AssetTypeActions_DialogueWave.h b/Engine/Source/Editor/AudioEditor/Private/AssetTypeActions/AssetTypeActions_DialogueWave.h
--- a/Engine/Source/Editor/AudioEditor/Private/AssetTypeActions/AssetTypeActions_DialogueWave.h
+++ b/Engine/Source/Editor/AudioEditor/Private/AssetTypeActions/AssetTypeActions_DialogueWave.h
@@ -16,6 +16,7 @@ public:
 	virtual FColor GetTypeColor() const override { return FColor(97, 85, 212); }
 	virtual UClass* GetSupportedClass() const override { return UDialogueWave::StaticClass(); }
 	virtual uint32 GetCategories() override { return EAssetTypeCategories::Sounds; }
+	virtual const TArray<FText>& GetSubMenus() const override;
 	virtual bool CanFilter() override { return true; }
 	virtual bool HasActions(const TArray<UObject*>& InObjects) const override { return true; }
 	virtual void GetActions(const TArray<UObject*>& InObjects, struct FToolMenuSection& Section) override;
diff --git a/Engine/Source/Editor/AudioEditor/Private/AssetTypeActions/AssetTypeActions_ReverbEffect.cpp b/Engine/Source/Editor/AudioEditor/Private/AssetTypeActions/AssetTypeActions_ReverbEffect.cpp
--- a/Engine/Source/Editor/AudioEditor/Private/AssetTypeActions/AssetTypeActions_ReverbEffect.cpp
+++ b/Engine/Source/Editor/AudioEditor/Private/AssetTypeActions/AssetTypeActions_ReverbEffect.cpp
@@ -3,7 +3,21 @@
 #include "AssetTypeActions/AssetTypeActions_ReverbEffect.h"
 #include "Sound/ReverbEffect.h"
 
+#define LOCTEXT_NAMESPACE "AssetTypeActions"
+
 UClass* FAssetTypeActions_ReverbEffect::GetSupportedClass() const
 {
 	return UReverbEffect::StaticClass();
 }
+
+const TArray<FText>& FAssetTypeActions_ReverbEffect::GetSubMenus() const
+{
+	static const TArray<FText> SubMenus
+	{
+		FText(LOCTEXT("AssetEffectSubMenu", "Effects"))
+	};
+
+	return SubMenus;
+}
+
+#undef LOCTEXT_NAMESPACE
\ No newline at end of file
diff --git a/Engine/Source/Editor/AudioEditor/Private/AssetTypeActions/AssetTypeActions_ReverbEffect.h b/Engine/Source/Editor/AudioEditor/Private/AssetTypeActions/AssetTypeActions_ReverbEffect.h
--- a/Engine/Source/Editor/AudioEditor/Private/AssetTypeActions/AssetTypeActions_ReverbEffect.h
+++ b/Engine/Source/Editor/AudioEditor/Private/AssetTypeActions/AssetTypeActions_ReverbEffect.h
@@ -12,5 +12,6 @@ public:
 	virtual FText GetName() const override { return NSLOCTEXT("AssetTypeActions", "AssetTypeActions_ReverbEffect", "Reverb Effect"); }
 	virtual FColor GetTypeColor() const override { return FColor(175, 0, 0); }
 	virtual UClass* GetSupportedClass() const override;
+	virtual const TArray<FText>& GetSubMenus() const override;
 	virtual uint32 GetCategories() override { return EAssetTypeCategories::Sounds; }
 };
diff --git a/Engine/Source/Editor/AudioEditor/Private/AssetTypeActions/AssetTypeActions_SoundClass.cpp b/Engine/Source/Editor/AudioEditor/Private/AssetTypeActions/AssetTypeActions_SoundClass.cpp
--- a/Engine/Source/Editor/AudioEditor/Private/AssetTypeActions/AssetTypeActions_SoundClass.cpp
+++ b/Engine/Source/Editor/AudioEditor/Private/AssetTypeActions/AssetTypeActions_SoundClass.cpp
@@ -11,6 +11,15 @@ UClass* FAssetTypeActions_SoundClass::GetSupportedClass() const
 	return USoundClass::StaticClass();
 }
 
+const TArray<FText>& FAssetTypeActions_SoundClass::GetSubMenus() const
+{
+	static const TArray<FText> SubMenus
+	{
+		LOCTEXT("AssetSoundClassSubMenu", "Classes")
+	};
+	return SubMenus;
+}
+
 void FAssetTypeActions_SoundClass::OpenAssetEditor( const TArray<UObject*>& InObjects, TSharedPtr<IToolkitHost> EditWithinLevelEditor )
 {
 	EToolkitMode::Type Mode = EditWithinLevelEditor.IsValid() ? EToolkitMode::WorldCentric : EToolkitMode::Standalone;
diff --git a/Engine/Source/Editor/AudioEditor/Private/AssetTypeActions/AssetTypeActions_SoundClass.h b/Engine/Source/Editor/AudioEditor/Private/AssetTypeActions/AssetTypeActions_SoundClass.h
--- a/Engine/Source/Editor/AudioEditor/Private/AssetTypeActions/AssetTypeActions_SoundClass.h
+++ b/Engine/Source/Editor/AudioEditor/Private/AssetTypeActions/AssetTypeActions_SoundClass.h
@@ -12,6 +12,7 @@ public:
 	// IAssetTypeActions Implementation
 	virtual FText GetName() const override { return NSLOCTEXT("AssetTypeActions", "AssetTypeActions_SoundClass", "Sound Class"); }
 	virtual FColor GetTypeColor() const override { return FColor(255, 175, 0); }
+	virtual const TArray<FText>& GetSubMenus() const override;
 	virtual UClass* GetSupportedClass() const override;
 	virtual void OpenAssetEditor( const TArray<UObject*>& InObjects, TSharedPtr<class IToolkitHost> EditWithinLevelEditor = TSharedPtr<IToolkitHost>() ) override;
 	virtual uint32 GetCategories() override { return EAssetTypeCategories::Sounds; }
diff --git a/Engine/Source/Editor/AudioEditor/Private/AssetTypeActions/AssetTypeActions_SoundEffectPreset.cpp b/Engine/Source/Editor/AudioEditor/Private/AssetTypeActions/AssetTypeActions_SoundEffectPreset.cpp
--- a/Engine/Source/Editor/AudioEditor/Private/AssetTypeActions/AssetTypeActions_SoundEffectPreset.cpp
+++ b/Engine/Source/Editor/AudioEditor/Private/AssetTypeActions/AssetTypeActions_SoundEffectPreset.cpp
@@ -8,6 +8,16 @@
 
 #include "Toolkits/SimpleAssetEditor.h"
 
+#define LOCTEXT_NAMESPACE "AssetTypeActions"
+
+namespace EffectPresets
+{
+	static const TArray<FText> SubMenus
+	{
+		FText(LOCTEXT("AssetEffectSubMenu", "Effects"))
+	};
+} // namespace EffectPresets
+
 FAssetTypeActions_SoundEffectPreset::FAssetTypeActions_SoundEffectPreset(USoundEffectPreset* InEffectPreset)
 	: EffectPreset(InEffectPreset)
 {
@@ -45,3 +55,25 @@ UClass* FAssetTypeActions_SoundEffectPreset::GetSupportedClass() const
 		return SupportedClass;
 	}
 }
+
+const TArray<FText>& FAssetTypeActions_SoundEffectSubmixPreset::GetSubMenus() const
+{
+	return EffectPresets::SubMenus;
+}
+
+const TArray<FText>& FAssetTypeActions_SoundEffectPreset::GetSubMenus() const
+{
+	return EffectPresets::SubMenus;
+}
+
+const TArray<FText>& FAssetTypeActions_SoundEffectSourcePresetChain::GetSubMenus() const
+{
+	return EffectPresets::SubMenus;
+}
+
+const TArray<FText>& FAssetTypeActions_SoundEffectSourcePreset::GetSubMenus() const
+{
+	return EffectPresets::SubMenus;
+}
+
+#undef LOCTEXT_NAMESPACE
\ No newline at end of file
diff --git a/Engine/Source/Editor/AudioEditor/Private/AssetTypeActions/AssetTypeActions_SoundEffectPreset.h b/Engine/Source/Editor/AudioEditor/Private/AssetTypeActions/AssetTypeActions_SoundEffectPreset.h
--- a/Engine/Source/Editor/AudioEditor/Private/AssetTypeActions/AssetTypeActions_SoundEffectPreset.h
+++ b/Engine/Source/Editor/AudioEditor/Private/AssetTypeActions/AssetTypeActions_SoundEffectPreset.h
@@ -20,6 +20,7 @@ public:
 	// IAssetTypeActions Implementation
 	virtual FText GetName() const override { return NSLOCTEXT("AssetTypeActions", "AssetTypeActions_SoundSubmixPreset", "Submix Effect Preset"); }
 	virtual FColor GetTypeColor() const override { return FColor(99, 63, 56); }
+	virtual const TArray<FText>& GetSubMenus() const override;
 	virtual UClass* GetSupportedClass() const override { return USoundEffectSubmixPreset::StaticClass(); }
 	virtual uint32 GetCategories() override { return EAssetTypeCategories::Sounds; }
 };
@@ -30,6 +31,7 @@ public:
 	// IAssetTypeActions Implementation
 	virtual FText GetName() const override { return NSLOCTEXT("AssetTypeActions", "AssetTypeActions_SoundSourcePreset", "Source Effect Preset"); }
 	virtual FColor GetTypeColor() const override { return FColor(72, 185, 187); }
+	virtual const TArray<FText>& GetSubMenus() const override;
 	virtual UClass* GetSupportedClass() const override { return USoundEffectSourcePreset::StaticClass(); }
 	virtual uint32 GetCategories() override { return EAssetTypeCategories::Sounds; }
 };
@@ -40,6 +42,7 @@ public:
 	// IAssetTypeActions Implementation
 	virtual FText GetName() const override { return NSLOCTEXT("AssetTypeActions", "AssetTypeActions_SoundSourcePresetChain", "Source Effect Preset Chain"); }
 	virtual FColor GetTypeColor() const override { return FColor(51, 107, 142); }
+	virtual const TArray<FText>& GetSubMenus() const override;
 	virtual UClass* GetSupportedClass() const override { return USoundEffectSourcePresetChain::StaticClass(); }
 	virtual uint32 GetCategories() override { return EAssetTypeCategories::Sounds; }
 };
@@ -52,6 +55,7 @@ public:
 	//~ Begin FAssetTypeActions_Base
 	virtual FText GetName() const override;
 	virtual FColor GetTypeColor() const override { return EffectPreset->GetPresetColor(); }
+	virtual const TArray<FText>& GetSubMenus() const override;
 	virtual UClass* GetSupportedClass() const override;
 	virtual uint32 GetCategories() override { return EAssetTypeCategories::Sounds; }
 	//~ End FAssetTypeActions_Base
diff --git a/Engine/Source/Editor/AudioEditor/Private/AssetTypeActions/AssetTypeActions_SoundMix.cpp b/Engine/Source/Editor/AudioEditor/Private/AssetTypeActions/AssetTypeActions_SoundMix.cpp
--- a/Engine/Source/Editor/AudioEditor/Private/AssetTypeActions/AssetTypeActions_SoundMix.cpp
+++ b/Engine/Source/Editor/AudioEditor/Private/AssetTypeActions/AssetTypeActions_SoundMix.cpp
@@ -3,7 +3,20 @@
 #include "AssetTypeActions/AssetTypeActions_SoundMix.h"
 #include "Sound/SoundMix.h"
 
+#define LOCTEXT_NAMESPACE "AssetTypeActions"
+
 UClass* FAssetTypeActions_SoundMix::GetSupportedClass() const
 {
 	return USoundMix::StaticClass();
 }
+
+const TArray<FText>& FAssetTypeActions_SoundMix::GetSubMenus() const
+{
+	static const TArray<FText> SubMenus
+	{
+		FText(LOCTEXT("AssetSoundClassSubMenu", "Classes"))
+	};
+
+	return SubMenus;
+}
+#undef LOCTEXT_NAMESPACE
\ No newline at end of file
diff --git a/Engine/Source/Editor/AudioEditor/Private/AssetTypeActions/AssetTypeActions_SoundMix.h b/Engine/Source/Editor/AudioEditor/Private/AssetTypeActions/AssetTypeActions_SoundMix.h
--- a/Engine/Source/Editor/AudioEditor/Private/AssetTypeActions/AssetTypeActions_SoundMix.h
+++ b/Engine/Source/Editor/AudioEditor/Private/AssetTypeActions/AssetTypeActions_SoundMix.h
@@ -9,8 +9,9 @@ class FAssetTypeActions_SoundMix : public FAssetTypeActions_Base
 {
 public:
 	// IAssetTypeActions Implementation
-	virtual FText GetName() const override { return NSLOCTEXT("AssetTypeActions", "AssetTypeActions_SoundMix", "Sound Mix"); }
+	virtual FText GetName() const override { return NSLOCTEXT("AssetTypeActions", "AssetTypeActions_SoundMix", "Sound Class Mix"); }
 	virtual FColor GetTypeColor() const override { return FColor(175, 0, 255); }
+	virtual const TArray<FText>& GetSubMenus() const override;
 	virtual UClass* GetSupportedClass() const override;
 	virtual uint32 GetCategories() override { return EAssetTypeCategories::Sounds; }
 };
diff --git a/Engine/Source/Editor/AudioEditor/Private/AssetTypeActions/AssetTypeActions_SoundSourceBus.cpp b/Engine/Source/Editor/AudioEditor/Private/AssetTypeActions/AssetTypeActions_SoundSourceBus.cpp
--- a/Engine/Source/Editor/AudioEditor/Private/AssetTypeActions/AssetTypeActions_SoundSourceBus.cpp
+++ b/Engine/Source/Editor/AudioEditor/Private/AssetTypeActions/AssetTypeActions_SoundSourceBus.cpp
@@ -14,9 +14,20 @@
 #include "AudioEditorModule.h"
 
 
+#define LOCTEXT_NAMESPACE "AssetTypeActions"
+
 UClass* FAssetTypeActions_SoundSourceBus::GetSupportedClass() const
 {
 	return USoundSourceBus::StaticClass();
 }
 
+const TArray<FText>& FAssetTypeActions_SoundSourceBus::GetSubMenus() const
+{
+	static const TArray<FText> SubMenus
+	{
+		FText(LOCTEXT("AssetSoundSourceSubMenu", "Source"))
+	};
 
+	return SubMenus;
+}
+#undef LOCTEXT_NAMESPACE
diff --git a/Engine/Source/Editor/AudioEditor/Private/AssetTypeActions/AssetTypeActions_SoundSourceBus.h b/Engine/Source/Editor/AudioEditor/Private/AssetTypeActions/AssetTypeActions_SoundSourceBus.h
--- a/Engine/Source/Editor/AudioEditor/Private/AssetTypeActions/AssetTypeActions_SoundSourceBus.h
+++ b/Engine/Source/Editor/AudioEditor/Private/AssetTypeActions/AssetTypeActions_SoundSourceBus.h
@@ -14,8 +14,9 @@ class FAssetTypeActions_SoundSourceBus : public FAssetTypeActions_SoundBase
 {
 public:
 	//~ Begin IAssetTypeActions Implementation
-	virtual FText GetName() const override { return NSLOCTEXT("AssetTypeActions", "AssetTypeActions_SoundSourceBus", "Sound Source Bus"); }
+	virtual FText GetName() const override { return NSLOCTEXT("AssetTypeActions", "AssetTypeActions_SoundSourceBus", "Source Bus"); }
 	virtual FColor GetTypeColor() const override { return FColor(212, 97, 85); }
+	virtual const TArray<FText>& GetSubMenus() const override;
 	virtual UClass* GetSupportedClass() const override;
 	virtual bool CanFilter() override { return true; }
 	virtual bool IsImportedAsset() const override { return false; }
diff --git a/Engine/Source/Editor/AudioEditor/Private/AssetTypeActions/AssetTypeActions_SoundSubmix.cpp b/Engine/Source/Editor/AudioEditor/Private/AssetTypeActions/AssetTypeActions_SoundSubmix.cpp
--- a/Engine/Source/Editor/AudioEditor/Private/AssetTypeActions/AssetTypeActions_SoundSubmix.cpp
+++ b/Engine/Source/Editor/AudioEditor/Private/AssetTypeActions/AssetTypeActions_SoundSubmix.cpp
@@ -26,4 +26,13 @@ void FAssetTypeActions_SoundSubmix::OpenAssetEditor( const TArray<UObject*>& InO
 	}
 }
 
+const TArray<FText>& FAssetTypeActions_SoundSubmix::GetSubMenus() const
+{
+	static const TArray<FText> SubMenus
+	{
+		FText(LOCTEXT("AssetSoundMixSubMenu", "Mix"))
+	};
+
+	return SubMenus;
+}
 #undef LOCTEXT_NAMESPACE
\ No newline at end of file
diff --git a/Engine/Source/Editor/AudioEditor/Private/AssetTypeActions/AssetTypeActions_SoundSubmix.h b/Engine/Source/Editor/AudioEditor/Private/AssetTypeActions/AssetTypeActions_SoundSubmix.h
--- a/Engine/Source/Editor/AudioEditor/Private/AssetTypeActions/AssetTypeActions_SoundSubmix.h
+++ b/Engine/Source/Editor/AudioEditor/Private/AssetTypeActions/AssetTypeActions_SoundSubmix.h
@@ -10,6 +10,7 @@ public:
 	// IAssetTypeActions Implementation
 	virtual FText GetName() const override { return NSLOCTEXT("AssetTypeActions", "AssetTypeActions_SoundSubmix", "Sound Submix"); }
 	virtual FColor GetTypeColor() const override { return FColor(175, 255, 0); }
+	virtual const TArray<FText>& GetSubMenus() const override;
 	virtual UClass* GetSupportedClass() const override;
 	virtual void OpenAssetEditor( const TArray<UObject*>& InObjects, TSharedPtr<class IToolkitHost> EditWithinLevelEditor = TSharedPtr<IToolkitHost>() ) override;
 	virtual uint32 GetCategories() override { return EAssetTypeCategories::Sounds; }
diff --git a/Engine/Source/Editor/AudioEditor/Private/AssetTypeActions/AssetTypeActions_SoundWave.cpp b/Engine/Source/Editor/AudioEditor/Private/AssetTypeActions/AssetTypeActions_SoundWave.cpp
--- a/Engine/Source/Editor/AudioEditor/Private/AssetTypeActions/AssetTypeActions_SoundWave.cpp
+++ b/Engine/Source/Editor/AudioEditor/Private/AssetTypeActions/AssetTypeActions_SoundWave.cpp
@@ -276,4 +276,13 @@ TSharedPtr<SWidget> FAssetTypeActions_SoundWave::GetThumbnailOverlay(const FAsse
 	return Box;
 }
 
+const TArray<FText>& FAssetTypeActions_SoundWave::GetSubMenus() const
+{
+	static const TArray<FText> SubMenus
+	{
+		FText(LOCTEXT("AssetSoundSourceSubMenu", "Source"))
+	};
+
+	return SubMenus;
+}
 #undef LOCTEXT_NAMESPACE
diff --git a/Engine/Source/Editor/AudioEditor/Private/AssetTypeActions/AssetTypeActions_SoundWave.h b/Engine/Source/Editor/AudioEditor/Private/AssetTypeActions/AssetTypeActions_SoundWave.h
--- a/Engine/Source/Editor/AudioEditor/Private/AssetTypeActions/AssetTypeActions_SoundWave.h
+++ b/Engine/Source/Editor/AudioEditor/Private/AssetTypeActions/AssetTypeActions_SoundWave.h
@@ -14,6 +14,7 @@ public:
 	// IAssetTypeActions Implementation
 	virtual FText GetName() const override { return NSLOCTEXT("AssetTypeActions", "AssetTypeActions_SoundWave", "Sound Wave"); }
 	virtual FColor GetTypeColor() const override { return FColor(97, 85, 212); }
+	virtual const TArray<FText>& GetSubMenus() const override;
 	virtual UClass* GetSupportedClass() const override;
 	virtual void GetActions(const TArray<UObject*>& InObjects, struct FToolMenuSection& Section) override;
 	virtual void OpenAssetEditor( const TArray<UObject*>& InObjects, TSharedPtr<class IToolkitHost> EditWithinLevelEditor = TSharedPtr<IToolkitHost>() ) override;
