Fix for multichannel file reimport in the case where file is reimported to be mono or stereo when it was previous multichannel.

Need to clear out the meta-data state of multichannel files.

diff --git a/Engine/Source/Editor/AudioEditor/Private/Factories/SoundFactory.cpp b/Engine/Source/Editor/AudioEditor/Private/Factories/SoundFactory.cpp
--- a/Engine/Source/Editor/AudioEditor/Private/Factories/SoundFactory.cpp
+++ b/Engine/Source/Editor/AudioEditor/Private/Factories/SoundFactory.cpp
@@ -226,6 +226,11 @@ UObject* USoundFactory::CreateObject
 			// Will block internally on audio thread completing outstanding commands
 			AudioDeviceManager->StopSoundsUsingResource(ExistingSound, &ComponentsToRestart);
 
+			// We need to clear out any stale multichannel data on the sound wave in the case this is a reimport from multichannel to mono/stereo
+			ExistingSound->ChannelOffsets.Reset();
+			ExistingSound->ChannelSizes.Reset();
+			ExistingSound->bIsAmbisonics = false;
+
 			// Resource data is required to exist, if it hasn't been loaded yet,
 			// to properly flush compressed data.  This allows the new version
 			// to be auditioned in the editor properly.
