build fix for reordering params to AddMenuEntry

diff --git a/Engine/Source/Editor/AudioEditor/Private/AssetTypeActions/AssetTypeActions_SoundBase.cpp b/Engine/Source/Editor/AudioEditor/Private/AssetTypeActions/AssetTypeActions_SoundBase.cpp
--- a/Engine/Source/Editor/AudioEditor/Private/AssetTypeActions/AssetTypeActions_SoundBase.cpp
+++ b/Engine/Source/Editor/AudioEditor/Private/AssetTypeActions/AssetTypeActions_SoundBase.cpp
@@ -48,6 +48,7 @@ void FAssetTypeActions_SoundBase::GetActions(const TArray<UObject*>& InObjects,
 		);
 
 	Section.AddMenuEntry(
+		"Sound_SoundMute",
 		LOCTEXT("Sound_MuteSound", "Mute"),
 		LOCTEXT("Sound_MuteSoundTooltip", "Mutes the selected sounds."),
 		FSlateIcon(FEditorStyle::GetStyleSetName(), "MediaAsset.AssetActions.Mute.Small"),
@@ -56,11 +57,11 @@ void FAssetTypeActions_SoundBase::GetActions(const TArray<UObject*>& InObjects,
 			FCanExecuteAction::CreateSP(this, &FAssetTypeActions_SoundBase::CanExecuteMuteCommand, Sounds),
 			FIsActionChecked::CreateSP(this, &FAssetTypeActions_SoundBase::IsActionCheckedMute, Sounds)
 		),
-		NAME_None, 
 		EUserInterfaceActionType::ToggleButton
 	);
 
 	Section.AddMenuEntry(
+		"Sound_StopSolo",
 		LOCTEXT("Sound_SoloSound", "Solo"),
 		LOCTEXT("Sound_SoloSoundTooltip", "Solos the selected sounds."),
 		FSlateIcon(FEditorStyle::GetStyleSetName(), "MediaAsset.AssetActions.Solo.Small"),
@@ -69,11 +70,10 @@ void FAssetTypeActions_SoundBase::GetActions(const TArray<UObject*>& InObjects,
 			FCanExecuteAction::CreateSP(this, &FAssetTypeActions_SoundBase::CanExecuteSoloCommand, Sounds),
 			FIsActionChecked::CreateSP(this, &FAssetTypeActions_SoundBase::IsActionCheckedSolo, Sounds)
 		),
-		NAME_None,
 		EUserInterfaceActionType::ToggleButton
 	);
-
 }
+
 PRAGMA_DISABLE_DEPRECATION_WARNINGS
 void FAssetTypeActions_SoundBase::AssetsActivated( const TArray<UObject*>& InObjects, EAssetTypeActivationMethod::Type ActivationType )
 {
