Protect pointers when audio is disabled (i.e. automation testing)

diff --git a/Engine/Source/Editor/AudioEditor/Private/AssetTypeActions/AssetTypeActions_SoundBase.cpp b/Engine/Source/Editor/AudioEditor/Private/AssetTypeActions/AssetTypeActions_SoundBase.cpp
--- a/Engine/Source/Editor/AudioEditor/Private/AssetTypeActions/AssetTypeActions_SoundBase.cpp
+++ b/Engine/Source/Editor/AudioEditor/Private/AssetTypeActions/AssetTypeActions_SoundBase.cpp
@@ -294,146 +294,164 @@ bool FAssetTypeActions_SoundBase::CanExecutePlayCommand(TArray<TWeakObjectPtr<US
 
 void FAssetTypeActions_SoundBase::ExecuteMuteSound(TArray<TWeakObjectPtr<USoundBase>> Objects) const
 {	
-	FAudioDebugger& Debugger = GEditor->GetAudioDeviceManager()->GetDebugger();	
-		
-	// In a selection that consists of some already muted, toggle everything in the same direction,
-	// to avoid AB problem.
-
-	bool bAnyMuted = IsActionCheckedMute(Objects);
-	for (TWeakObjectPtr<USoundBase> SoundBase : Objects)
+	if (FAudioDeviceManager* ADM = GEditor->GetAudioDeviceManager())
 	{
-		if (USoundCue* Cue = Cast<USoundCue>(SoundBase.Get()))
-		{											
-			Debugger.SetMuteSoundCue(Cue->GetFName(), !bAnyMuted);
-		}
-		else if(USoundWave* Wave = Cast<USoundWave>(SoundBase.Get()))
+		FAudioDebugger& Debugger = ADM->GetDebugger();
+
+		// In a selection that consists of some already muted, toggle everything in the same direction,
+		// to avoid AB problem.
+
+		bool bAnyMuted = IsActionCheckedMute(Objects);
+		for (TWeakObjectPtr<USoundBase> SoundBase : Objects)
 		{
-			Debugger.SetMuteSoundWave(Wave->GetFName(), !bAnyMuted);
-		}		
-	}	
+			if (USoundCue* Cue = Cast<USoundCue>(SoundBase.Get()))
+			{
+				Debugger.SetMuteSoundCue(Cue->GetFName(), !bAnyMuted);
+			}
+			else if (USoundWave* Wave = Cast<USoundWave>(SoundBase.Get()))
+			{
+				Debugger.SetMuteSoundWave(Wave->GetFName(), !bAnyMuted);
+			}
+		}
+	}
 }
 
 void FAssetTypeActions_SoundBase::ExecuteSoloSound(TArray<TWeakObjectPtr<USoundBase>> Objects) const
 {
-	FAudioDebugger& Debugger = GEditor->GetAudioDeviceManager()->GetDebugger();
-	
-	// In a selection that consists of some already soloed, toggle everything in the same direction,
-	// to avoid AB problem.
-		
-	bool bAnySoloed = IsActionCheckedSolo(Objects);
-	for (TWeakObjectPtr<USoundBase> SoundBase : Objects)
+	if (FAudioDeviceManager* ADM = GEditor->GetAudioDeviceManager())
 	{
-		if (USoundCue* Cue = Cast<USoundCue>(SoundBase.Get()))
-		{
-			Debugger.SetSoloSoundCue(Cue->GetFName(), !bAnySoloed);
-		}
-		else if (USoundWave* Wave = Cast<USoundWave>(SoundBase.Get()))
+		FAudioDebugger& Debugger = ADM->GetDebugger();
+
+		// In a selection that consists of some already soloed, toggle everything in the same direction,
+		// to avoid AB problem.
+
+		bool bAnySoloed = IsActionCheckedSolo(Objects);
+		for (TWeakObjectPtr<USoundBase> SoundBase : Objects)
 		{
-			Debugger.SetSoloSoundCue(Wave->GetFName(), !bAnySoloed);
+			if (USoundCue* Cue = Cast<USoundCue>(SoundBase.Get()))
+			{
+				Debugger.SetSoloSoundCue(Cue->GetFName(), !bAnySoloed);
+			}
+			else if (USoundWave* Wave = Cast<USoundWave>(SoundBase.Get()))
+			{
+				Debugger.SetSoloSoundCue(Wave->GetFName(), !bAnySoloed);
+			}
 		}
 	}
 }
 
 bool FAssetTypeActions_SoundBase::IsActionCheckedMute(TArray<TWeakObjectPtr<USoundBase>> Objects) const
 {
-	// If *any* of the selection are muted, show the tick box as ticked.
-	FAudioDebugger& Debugger = GEditor->GetAudioDeviceManager()->GetDebugger();
-	for (TWeakObjectPtr<USoundBase> SoundBase : Objects)
+	if (FAudioDeviceManager* ADM = GEditor->GetAudioDeviceManager())
 	{
-		if (USoundCue* Cue = Cast<USoundCue>(SoundBase.Get()))
+		// If *any* of the selection are muted, show the tick box as ticked.
+		FAudioDebugger& Debugger = ADM->GetDebugger();
+		for (TWeakObjectPtr<USoundBase> SoundBase : Objects)
 		{
-			if (Debugger.IsMuteSoundCue(Cue->GetFName()))
+			if (USoundCue* Cue = Cast<USoundCue>(SoundBase.Get()))
 			{
-				return true;
+				if (Debugger.IsMuteSoundCue(Cue->GetFName()))
+				{
+					return true;
+				}
 			}
-		}
-		else if (USoundWave* Wave = Cast<USoundWave>(SoundBase.Get()))
-		{
-			if (Debugger.IsMuteSoundWave(Wave->GetFName()))
+			else if (USoundWave* Wave = Cast<USoundWave>(SoundBase.Get()))
 			{
-				return true;
+				if (Debugger.IsMuteSoundWave(Wave->GetFName()))
+				{
+					return true;
+				}
 			}
 		}
 	}
-
 	return false;
 }
 
 bool FAssetTypeActions_SoundBase::IsActionCheckedSolo(TArray<TWeakObjectPtr<USoundBase>> Objects) const
 {
 	// If *any* of the selection are solod, show the tick box as ticked.
-	FAudioDebugger& Debugger = GEditor->GetAudioDeviceManager()->GetDebugger();
-	for (TWeakObjectPtr<USoundBase> SoundBase : Objects)
+	if (FAudioDeviceManager* ADM = GEditor->GetAudioDeviceManager())
 	{
-		if (USoundCue* Cue = Cast<USoundCue>(SoundBase.Get()))
+		FAudioDebugger& Debugger = ADM->GetDebugger();
+		for (TWeakObjectPtr<USoundBase> SoundBase : Objects)
 		{
-			if (Debugger.IsSoloSoundCue(Cue->GetFName()))
+			if (USoundCue* Cue = Cast<USoundCue>(SoundBase.Get()))
 			{
-				return true;
+				if (Debugger.IsSoloSoundCue(Cue->GetFName()))
+				{
+					return true;
+				}
 			}
-		}
-		else if (USoundWave* Wave = Cast<USoundWave>(SoundBase.Get()))
-		{
-			if (Debugger.IsSoloSoundWave(Wave->GetFName()))
+			else if (USoundWave* Wave = Cast<USoundWave>(SoundBase.Get()))
 			{
-				return true;
+				if (Debugger.IsSoloSoundWave(Wave->GetFName()))
+				{
+					return true;
+				}
 			}
 		}
 	}
-
 	return false;
 }
 
 bool FAssetTypeActions_SoundBase::CanExecuteMuteCommand(TArray<TWeakObjectPtr<USoundBase>> Objects) const
 {
-	// Allow muting if we're not Soloing.
-	FAudioDebugger& Debugger = GEditor->GetAudioDeviceManager()->GetDebugger();
-	for (TWeakObjectPtr<USoundBase> SoundBase : Objects)
-	{		
-		if (USoundCue* Cue = Cast<USoundCue>(SoundBase.Get()))
+	if (FAudioDeviceManager* ADM = GEditor->GetAudioDeviceManager())
+	{
+		// Allow muting if we're not Soloing.
+		FAudioDebugger& Debugger = ADM->GetDebugger();
+		for (TWeakObjectPtr<USoundBase> SoundBase : Objects)
 		{
-			if (Debugger.IsSoloSoundCue(Cue->GetFName()))
+			if (USoundCue* Cue = Cast<USoundCue>(SoundBase.Get()))
 			{
-				return false;
+				if (Debugger.IsSoloSoundCue(Cue->GetFName()))
+				{
+					return false;
+				}
 			}
-		}
-		else if (USoundWave* Wave = Cast<USoundWave>(SoundBase.Get()))
-		{
-			if (Debugger.IsSoloSoundWave(Wave->GetFName()))
+			else if (USoundWave* Wave = Cast<USoundWave>(SoundBase.Get()))
 			{
-				return false;
+				if (Debugger.IsSoloSoundWave(Wave->GetFName()))
+				{
+					return false;
+				}
 			}
 		}
-	}
 
-	// Ok.
-	return true;
+		// Ok.
+		return true;
+	}
+	return false;
 }
 
 bool FAssetTypeActions_SoundBase::CanExecuteSoloCommand(TArray<TWeakObjectPtr<USoundBase>> Objects) const
 {	
-	// Allow Soloing if we're not Muting.
-	FAudioDebugger& Debugger = GEditor->GetAudioDeviceManager()->GetDebugger();
-	for (TWeakObjectPtr<USoundBase> SoundBase : Objects)
-	{		
-		if (USoundCue* Cue = Cast<USoundCue>(SoundBase.Get()))
+	if (FAudioDeviceManager* ADM = GEditor->GetAudioDeviceManager())
+	{
+		// Allow Soloing if we're not Muting.
+		FAudioDebugger& Debugger = ADM->GetDebugger();
+		for (TWeakObjectPtr<USoundBase> SoundBase : Objects)
 		{
-			if (Debugger.IsMuteSoundCue(Cue->GetFName()))
+			if (USoundCue* Cue = Cast<USoundCue>(SoundBase.Get()))
 			{
-				return false;
+				if (Debugger.IsMuteSoundCue(Cue->GetFName()))
+				{
+					return false;
+				}
 			}
-		}
-		else if (USoundWave* Wave = Cast<USoundWave>(SoundBase.Get()))
-		{
-			if (Debugger.IsMuteSoundWave(Wave->GetFName()))
+			else if (USoundWave* Wave = Cast<USoundWave>(SoundBase.Get()))
 			{
-				return false;
+				if (Debugger.IsMuteSoundWave(Wave->GetFName()))
+				{
+					return false;
+				}
 			}
 		}
-	}
 
-	// Ok.
-	return true;
+		// Ok.
+		return true;
+	}
+	return false;
 }
 
 #undef LOCTEXT_NAMESPACE
diff --git a/Engine/Source/Editor/AudioEditor/Private/AssetTypeActions/AssetTypeActions_SoundClass.cpp b/Engine/Source/Editor/AudioEditor/Private/AssetTypeActions/AssetTypeActions_SoundClass.cpp
--- a/Engine/Source/Editor/AudioEditor/Private/AssetTypeActions/AssetTypeActions_SoundClass.cpp
+++ b/Engine/Source/Editor/AudioEditor/Private/AssetTypeActions/AssetTypeActions_SoundClass.cpp
@@ -73,38 +73,44 @@ void FAssetTypeActions_SoundClass::GetActions(const TArray<UObject*>& InObjects,
 
 void FAssetTypeActions_SoundClass::ExecuteMute(TArray<TWeakObjectPtr<USoundClass>> Objects) const
 {
-	FAudioDebugger& Debugger = GEditor->GetAudioDeviceManager()->GetDebugger();
-	for (TWeakObjectPtr<USoundClass> i : Objects)
-	{		
-		if (USoundClass* Class = Cast<USoundClass>(i.Get()))
+	if (FAudioDeviceManager* ADM = GEditor->GetAudioDeviceManager())
+	{
+		for (TWeakObjectPtr<USoundClass> i : Objects)
 		{
-			Debugger.ToggleMuteSoundClass(Class->GetFName());
+			if (USoundClass* Class = Cast<USoundClass>(i.Get()))
+			{
+				ADM->GetDebugger().ToggleMuteSoundClass(Class->GetFName());
+			}
 		}
 	}
 }
 
 void FAssetTypeActions_SoundClass::ExecuteSolo(TArray<TWeakObjectPtr<USoundClass>> Objects) const
 {
-	FAudioDebugger& Debugger = GEditor->GetAudioDeviceManager()->GetDebugger();
-	for (TWeakObjectPtr<USoundClass> i : Objects)
+	if (FAudioDeviceManager* ADM = GEditor->GetAudioDeviceManager())
 	{
-		if (USoundClass* Class = Cast<USoundClass>(i.Get()))
+		for (TWeakObjectPtr<USoundClass> i : Objects)
 		{
-			Debugger.ToggleSoloSoundClass(Class->GetFName());
+			if (USoundClass* Class = Cast<USoundClass>(i.Get()))
+			{
+				ADM->GetDebugger().ToggleSoloSoundClass(Class->GetFName());
+			}
 		}
 	}
 }
 
 bool FAssetTypeActions_SoundClass::IsActionCheckedMute(TArray<TWeakObjectPtr<USoundClass>> Objects) const
 {
-	FAudioDebugger& Debugger = GEditor->GetAudioDeviceManager()->GetDebugger();
-	for (TWeakObjectPtr<USoundClass> i : Objects)
+	if (FAudioDeviceManager* ADM = GEditor->GetAudioDeviceManager())
 	{
-		if (USoundClass* Class = Cast<USoundClass>(i.Get()))
+		for (TWeakObjectPtr<USoundClass> i : Objects)
 		{
-			if (Debugger.IsMuteSoundClass(Class->GetFName()))
+			if (USoundClass* Class = Cast<USoundClass>(i.Get()))
 			{
-				return true;
+				if (ADM->GetDebugger().IsMuteSoundClass(Class->GetFName()))
+				{
+					return true;
+				}
 			}
 		}
 	}
@@ -113,60 +119,62 @@ bool FAssetTypeActions_SoundClass::IsActionCheckedMute(TArray<TWeakObjectPtr<USo
 
 bool FAssetTypeActions_SoundClass::IsActionCheckedSolo(TArray<TWeakObjectPtr<USoundClass>> Objects) const
 {
-	FAudioDebugger& Debugger = GEditor->GetAudioDeviceManager()->GetDebugger();
-
-	for (TWeakObjectPtr<USoundClass> i : Objects)
-	{
-		if (USoundClass* Class = Cast<USoundClass>(i.Get()))
+	if (FAudioDeviceManager* ADM = GEditor->GetAudioDeviceManager())
+	{		
+		for (TWeakObjectPtr<USoundClass> i : Objects)
 		{
-			if (Debugger.IsSoloSoundClass(Class->GetFName()))
+			if (USoundClass* Class = Cast<USoundClass>(i.Get()))
 			{
-				return true;
+				if (ADM->GetDebugger().IsSoloSoundClass(Class->GetFName()))
+				{
+					return true;
+				}
 			}
 		}
 	}
-
 	return false;
 }
 
 bool FAssetTypeActions_SoundClass::CanExecuteMuteCommand(TArray<TWeakObjectPtr<USoundClass>> Objects) const
 {
-	FAudioDebugger& Debugger = GEditor->GetAudioDeviceManager()->GetDebugger();
-
-	// Allow muting if we're not Soloing.
-	for (TWeakObjectPtr<USoundClass> i : Objects)
+	if (FAudioDeviceManager* ADM = GEditor->GetAudioDeviceManager())
 	{
-		if (USoundClass* Class = Cast<USoundClass>(i.Get()))
+		// Allow muting if we're not Soloing.
+		for (TWeakObjectPtr<USoundClass> i : Objects)
 		{
-			if (Debugger.IsSoloSoundClass(Class->GetFName()))
+			if (USoundClass* Class = Cast<USoundClass>(i.Get()))
 			{
-				return false;
+				if (ADM->GetDebugger().IsSoloSoundClass(Class->GetFName()))
+				{
+					return false;
+				}
 			}
 		}
+		// Ok.
+		return true;
 	}
-
-	// Ok.
-	return true;
+	return false;
 }
 
 bool FAssetTypeActions_SoundClass::CanExecuteSoloCommand(TArray<TWeakObjectPtr<USoundClass>> Objects) const
 {
-	FAudioDebugger& Debugger = GEditor->GetAudioDeviceManager()->GetDebugger();
-
-	// Allow Soloing if we're not Muting.
-	for (TWeakObjectPtr<USoundClass> i : Objects)
+	if (FAudioDeviceManager* ADM = GEditor->GetAudioDeviceManager())
 	{
-		if (USoundClass* Class = Cast<USoundClass>(i.Get()))
+		// Allow Soloing if we're not Muting.
+		for (TWeakObjectPtr<USoundClass> i : Objects)
 		{
-			if (Debugger.IsMuteSoundClass(Class->GetFName()))
+			if (USoundClass* Class = Cast<USoundClass>(i.Get()))
 			{
-				return false;
+				if (ADM->GetDebugger().IsMuteSoundClass(Class->GetFName()))
+				{
+					return false;
+				}
 			}
 		}
+		// Ok.
+		return true;
 	}
-
-	// Ok.
-	return true;
+	return false;
 }
 
 
diff --git a/Engine/Source/Editor/AudioEditor/Private/SoundClassEditor.cpp b/Engine/Source/Editor/AudioEditor/Private/SoundClassEditor.cpp
--- a/Engine/Source/Editor/AudioEditor/Private/SoundClassEditor.cpp
+++ b/Engine/Source/Editor/AudioEditor/Private/SoundClassEditor.cpp
@@ -161,9 +161,10 @@ void FSoundClassEditor::InitSoundClassEditor( const EToolkitMode::Type Mode, con
 	ExtendToolbar();
 	RegenerateMenusAndToolbars();
 			
-	ensure(GEditor);
-	ensure(GEditor->GetAudioDeviceManager());
-	Debugger = &GEditor->GetAudioDeviceManager()->GetDebugger();
+	if (GEditor->GetAudioDeviceManager())
+	{
+		Debugger = &GEditor->GetAudioDeviceManager()->GetDebugger();
+	}	
 
 	GraphEditor->SelectAllNodes();
 	for (UObject* SelectedNode : GraphEditor->GetSelectedNodes())
@@ -395,34 +396,40 @@ void FSoundClassEditor::RedoGraphAction()
 
 void FSoundClassEditor::ToggleSolo()
 {
-	Debugger->ToggleSoloSoundClass(SoundClass->GetFName());
+	if (Debugger)
+	{
+		Debugger->ToggleSoloSoundClass(SoundClass->GetFName());
+	}
 }
 
 bool FSoundClassEditor::CanExcuteToggleSolo() const
 {
 	// Enable solo if we are not Muted	
-	return !Debugger->IsMuteSoundClass(SoundClass->GetFName());
+	return Debugger ? !Debugger->IsMuteSoundClass(SoundClass->GetFName()) : false;
 }
 
 bool FSoundClassEditor::IsSoloToggled() const
 {
-	return Debugger->IsSoloSoundClass(SoundClass->GetFName());
+	return Debugger ? Debugger->IsSoloSoundClass(SoundClass->GetFName()) : false;
 }
 
 void FSoundClassEditor::ToggleMute()
 {
-	Debugger->ToggleMuteSoundClass(SoundClass->GetFName());
+	if (Debugger)
+	{
+		Debugger->ToggleMuteSoundClass(SoundClass->GetFName());
+	}
 }
 
 bool FSoundClassEditor::CanExcuteToggleMute() const
 {
 	// Enable mute if we are not Soloed
-	return !Debugger->IsSoloSoundClass(SoundClass->GetFName());
+	return Debugger ? !Debugger->IsSoloSoundClass(SoundClass->GetFName()) : false;
 }
 
 bool FSoundClassEditor::IsMuteToggled() const
 {	
-	return Debugger->IsMuteSoundClass(SoundClass->GetFName());
+	return Debugger ? Debugger->IsMuteSoundClass(SoundClass->GetFName()) : false;
 }
 
 void FSoundClassEditor::CreateSoundClass(UEdGraphPin* FromPin, const FVector2D& Location, const FString& Name)
diff --git a/Engine/Source/Editor/AudioEditor/Private/SoundCueEditor.cpp b/Engine/Source/Editor/AudioEditor/Private/SoundCueEditor.cpp
--- a/Engine/Source/Editor/AudioEditor/Private/SoundCueEditor.cpp
+++ b/Engine/Source/Editor/AudioEditor/Private/SoundCueEditor.cpp
@@ -149,10 +149,11 @@ void FSoundCueEditor::InitSoundCueEditor(const EToolkitMode::Type Mode, const TS
 	ExtendToolbar();
 	RegenerateMenusAndToolbars();
 
-	check(GEditor->GetAudioDeviceManager());
-	Debugger = &GEditor->GetAudioDeviceManager()->GetDebugger();
-	check(Debugger);
-
+	if (GEditor->GetAudioDeviceManager())
+	{
+		Debugger = &GEditor->GetAudioDeviceManager()->GetDebugger();
+	}	
+	
 	// @todo toolkit world centric editing
 	/*if(IsWorldCentricAssetEditor())
 	{
@@ -421,34 +422,40 @@ void FSoundCueEditor::TogglePlayback()
 
 void FSoundCueEditor::ToggleSolo()
 {
-	Debugger->ToggleSoloSoundCue(SoundCue->GetFName());
+	if (Debugger)
+	{
+		Debugger->ToggleSoloSoundCue(SoundCue->GetFName());
+	}
 }
 
 bool FSoundCueEditor::CanExcuteToggleSolo() const
 {
 	// Allow Solo if Mute is not Toggle on
-	return !Debugger->IsMuteSoundCue(SoundCue->GetFName());
+	return Debugger ? !Debugger->IsMuteSoundCue(SoundCue->GetFName()) : false;
 }
 
 bool FSoundCueEditor::IsSoloToggled() const
 {
-	return Debugger->IsSoloSoundCue(SoundCue->GetFName());
+	return Debugger ? Debugger->IsSoloSoundCue(SoundCue->GetFName()) : false;
 }
 
 void FSoundCueEditor::ToggleMute()
 {
-	Debugger->ToggleMuteSoundCue(SoundCue->GetFName());	
+	if (Debugger)
+	{
+		Debugger->ToggleMuteSoundCue(SoundCue->GetFName());
+	}
 }
 
 bool FSoundCueEditor::CanExcuteToggleMute() const
 {
 	// Allow Mute if Solo is not Toggle on
-	return !Debugger->IsSoloSoundCue(SoundCue->GetFName());
+	return Debugger ? !Debugger->IsSoloSoundCue(SoundCue->GetFName()) : false;
 }
 
 bool FSoundCueEditor::IsMuteToggled() const
 {
-	return Debugger->IsMuteSoundCue(SoundCue->GetFName());
+	return Debugger ? Debugger->IsMuteSoundCue(SoundCue->GetFName()) : false;
 }
 
 void FSoundCueEditor::PlaySingleNode(UEdGraphNode* Node)
diff --git a/Engine/Source/Runtime/Engine/Private/AudioDebug.cpp b/Engine/Source/Runtime/Engine/Private/AudioDebug.cpp
--- a/Engine/Source/Runtime/Engine/Private/AudioDebug.cpp
+++ b/Engine/Source/Runtime/Engine/Private/AudioDebug.cpp
@@ -164,7 +164,7 @@ namespace
 
 	void HandleDumpActiveSounds(UWorld* World)
 	{
-		if (GEngine)
+		if (GEngine && GEngine->GetAudioDeviceManager())
 		{
 			GEngine->GetAudioDeviceManager()->GetDebugger().DumpActiveSounds();
 		}
@@ -172,7 +172,7 @@ namespace
 	
 	void HandleClearMutesAndSolos(UWorld* World)
 	{
-		if (GEngine)
+		if (GEngine && GEngine->GetAudioDeviceManager())
 		{
 			GEngine->GetAudioDeviceManager()->GetDebugger().ClearMutesAndSolos();
 		}
@@ -1318,8 +1318,6 @@ void FAudioDebugger::SendUpdateResultsToGameThread(const FAudioDevice& AudioDevi
 
 	const uint8 RequestedStats = Stats_AudioThread->RequestedStats;
 	TMap<FActiveSound*, int32> ActiveSoundToInfoIndex;
-
-	FAudioDebugger& Debugger = GEngine->GetAudioDeviceManager()->GetDebugger();
 	
 	const bool bDebug = (RequestedStats & ERequestedAudioStats::DebugSounds) != 0;
 
