Clean-up SoundSubmix Implementation
- Make Submix graph point in an intuitive direction and change input/output directions.
- Fix bugs around adding/removing submixes and having multiple editors open.
- Move implementation and definitions so SoundSubmix updates builds faster (remove SoundSubmix.h references from bulky headers that are used all over the engine)

diff --git a/Engine/Plugins/Runtime/ResonanceAudio/Source/ResonanceAudio/Private/ResonanceAudioReverb.cpp b/Engine/Plugins/Runtime/ResonanceAudio/Source/ResonanceAudio/Private/ResonanceAudioReverb.cpp
--- a/Engine/Plugins/Runtime/ResonanceAudio/Source/ResonanceAudio/Private/ResonanceAudioReverb.cpp
+++ b/Engine/Plugins/Runtime/ResonanceAudio/Source/ResonanceAudio/Private/ResonanceAudioReverb.cpp
@@ -7,6 +7,7 @@
 #include "ResonanceAudioCommon.h"
 #include "ResonanceAudioModule.h"
 #include "ResonanceAudioSettings.h"
+#include "Sound/SoundSubmix.h"
 
 namespace ResonanceAudio
 {
diff --git a/Engine/Plugins/Runtime/Synthesis/Source/Synthesis/Classes/SubmixEffects/SubmixEffectTapDelay.h b/Engine/Plugins/Runtime/Synthesis/Source/Synthesis/Classes/SubmixEffects/SubmixEffectTapDelay.h
--- a/Engine/Plugins/Runtime/Synthesis/Source/Synthesis/Classes/SubmixEffects/SubmixEffectTapDelay.h
+++ b/Engine/Plugins/Runtime/Synthesis/Source/Synthesis/Classes/SubmixEffects/SubmixEffectTapDelay.h
@@ -60,13 +60,13 @@ struct SYNTHESIS_API FTapDelayInfo
 	UPROPERTY(EditAnywhere, BlueprintReadWrite, Category = Realtime, meta = (ClampMin = "-60.0", ClampMax = "6.0"))
 	float Gain;
 
-	// When the Tap Line Mode is set to Send To Channel, this parameter designates which channel the echo should play out of.
-	UPROPERTY(EditAnywhere, BlueprintReadWrite, Category = Realtime)
+	// When the Tap Line Mode is set to Send To Channel, designates index of channel from which the echo should play.
+	UPROPERTY(EditAnywhere, BlueprintReadWrite, Category = Realtime, meta = (EditCondition="TapLineMode == ETapLineMode::SendToChannel", ClampMin = "0", ClampMax = "8"))
 	int32 OutputChannel;
 
-	// When the Tap Line Mode is set to Panning, this parameter designates the angle at which the echo should be panned.
+	// When the Tap Line Mode is set to Panning, designates the angle at which the echo should be panned.
 	// -90 is left, 90 is right, and 180/-180 is directly behind the listener.
-	UPROPERTY(EditAnywhere, BlueprintReadWrite, Category = Realtime, meta = (ClampMin = "-180.0", ClampMax = "180.0"))
+	UPROPERTY(EditAnywhere, BlueprintReadWrite, Category = Realtime, meta = (EditCondition = "TapLineMode == ETapLineMode::Panning", ClampMin = "-180.0", ClampMax = "180.0"))
 	float PanInDegrees;
 
 	UPROPERTY(Transient, meta = (IgnoreForMemberInitializationTest))
@@ -171,7 +171,7 @@ private:
 // Class which processes audio streams and uses parameters defined in the preset class.
 // ========================================================================
 
-UCLASS(ClassGroup = AudioSourceEffect, meta = (BlueprintSpawnableComponent))
+UCLASS(ClassGroup = AudioSourceEffect, meta = (BlueprintSpawnableComponent), AutoExpandCategories = (Realtime))
 class SYNTHESIS_API USubmixEffectTapDelayPreset : public USoundEffectSubmixPreset
 {
 	GENERATED_BODY()
diff --git a/Engine/Source/Editor/AudioEditor/Classes/SoundSubmixGraph/SoundSubmixGraph.h b/Engine/Source/Editor/AudioEditor/Classes/SoundSubmixGraph/SoundSubmixGraph.h
--- a/Engine/Source/Editor/AudioEditor/Classes/SoundSubmixGraph/SoundSubmixGraph.h
+++ b/Engine/Source/Editor/AudioEditor/Classes/SoundSubmixGraph/SoundSubmixGraph.h
@@ -1,23 +1,30 @@
 // Copyright 1998-2019 Epic Games, Inc. All Rights Reserved.
-
 #pragma once
+
 #include "EdGraph/EdGraph.h"
 #include "SoundSubmixGraph.generated.h"
 
+
+// Forward Declarations
+class USoundSubmix;
+class USoundSubmixGraphNode;
+
+
 UCLASS(MinimalAPI)
 class USoundSubmixGraph : public UEdGraph
 {
 	GENERATED_UCLASS_BODY()
 
+public:
 	/**
 	 * Set the SoundSubmix which forms the root of this graph
 	 */
-	AUDIOEDITOR_API void SetRootSoundSubmix(class USoundSubmix* InSoundSubmix);
+	AUDIOEDITOR_API void SetRootSoundSubmix(USoundSubmix* InSoundSubmix);
 
 	/**
 	 * Get the SoundSubmix which forms the root of this graph
 	 */
-	class USoundSubmix* GetRootSoundSubmix() const;
+	USoundSubmix* GetRootSoundSubmix() const;
 
 	/**
 	 * Completely rebuild the graph from the root, removing all old nodes
@@ -28,37 +35,47 @@ class USoundSubmixGraph : public UEdGraph
 	 * Display SoundSubmixes (and all of their children) that have been dragged onto the editor
 	 *
 	 * @param	SoundSubmixes	SoundSubmixes not already represented on the graph
-	 * @param	NodePosX		X coordinate classes were dropped at
-	 * @param	NodePosY		Y coordinate classes were dropped at
+	 * @param	NodePosX		X coordinate submix(es) were dropped at
+	 * @param	NodePosY		Y coordinate submix(es) were dropped at
 	 */
-	void AddDroppedSoundSubmixes(const TArray<class USoundSubmix*>& SoundSubmixes, int32 NodePosX, int32 NodePosY);
+	void AddDroppedSoundSubmixes(const TArray<USoundSubmix*>& SoundSubmixes, int32 NodePosX, int32 NodePosY);
 
 	/**
 	 * Display a new SoundSubmix that has just been created using the editor
 	 *
 	 * @param	FromPin		The Pin that was dragged from to create the SoundSubmix (may be NULL)
 	 * @param	SoundClass	The newly created SoundSubmix
-	 * @param	NodePosX	X coordinate new class was created at
-	 * @param	NodePosY	Y coordinate new class was created at
+	 * @param	NodePosX	X coordinate new submix was created at
+	 * @param	NodePosY	Y coordinate new submix was created at
 	 * @param bSelectNewNode	Whether or not to select the new node being created
 	 */
-	AUDIOEDITOR_API void AddNewSoundSubmix(UEdGraphPin* FromPin, class USoundSubmix* SoundSubmix, int32 NodePosX, int32 NodePosY, bool bSelectNewNode = true);
+	AUDIOEDITOR_API void AddNewSoundSubmix(UEdGraphPin* FromPin, USoundSubmix* SoundSubmix, int32 NodePosX, int32 NodePosY, bool bSelectNewNode = true);
+
 	/**
 	 * Checks whether a SoundSubmix is already represented on this graph
 	 */
-	bool IsSubmixDisplayed(class USoundSubmix* SoundSubmix) const;
+	bool IsSubmixDisplayed(USoundSubmix* SoundSubmix) const;
+
 	/**
 	 * Use this graph to re-link all of the SoundSubmixes it represents after a change in linkage
 	 */
 	void LinkSoundSubmixes();
+
 	/**
 	 * Re-link all of the nodes in this graph after a change to SoundSubmix linkage
 	 */
 	AUDIOEDITOR_API void RefreshGraphLinks();
+
 	/**
 	 * Recursively remove a set of nodes from this graph and re-link SoundSubmixes afterwards
 	 */
-	AUDIOEDITOR_API void RecursivelyRemoveNodes(const TSet<class UObject*> NodesToRemove);
+	AUDIOEDITOR_API void RecursivelyRemoveNodes(const TSet<UObject*> NodesToRemove);
+
+
+	/**
+	 * Find an existing node that represents a given SoundSubmix
+	 */
+	AUDIOEDITOR_API USoundSubmixGraphNode* FindExistingNode(USoundSubmix* SoundSubmix) const;
 
 private:
 	/**
@@ -70,7 +87,7 @@ private:
 	 * @param bSelectNewNode	Whether or not to select the new node being created
 	 * @return	Total height of all constructed nodes (used to arrange multiple new nodes)
 	 */
-	int32 ConstructNodes(class USoundSubmix* SoundSubmix, int32 NodePosX, int32 NodePosY, bool bSelectNewNode = true);
+	int32 ConstructNodes(USoundSubmix* SoundSubmix, int32 NodePosX, int32 NodePosY, bool bSelectNewNode = true);
 	/**
 	 * Recursively build a map of child counts for each SoundSubmix to arrange them correctly
 	 *
@@ -78,7 +95,7 @@ private:
 	 * @param	OutChildCounts	Map of child counts
 	 * @return	Total child count for ParentClass
 	 */
-	int32 RecursivelyGatherChildCounts(class USoundSubmix* ParentSubmix, TMap<class USoundSubmix*, int32>& OutChildCounts);
+	int32 RecursivelyGatherChildCounts(USoundSubmix* ParentSubmix, TMap<USoundSubmix*, int32>& OutChildCounts);
 	/**
 	 * Recursively Construct Nodes to represent the children of a SoundSubmix
 	 *
@@ -87,11 +104,11 @@ private:
 	 * @param bSelectNewNode	Whether or not to select the new node being created
 	 * @return	Total height of constructed nodes (used to arrange next new node)
 	 */
-	int32 RecursivelyConstructChildNodes(class USoundSubmixGraphNode* ParentNode, const TMap<class USoundSubmix*, int32>& InChildCounts, bool bSelectNewNode = true);
+	int32 RecursivelyConstructChildNodes(USoundSubmixGraphNode* ParentNode, const TMap<USoundSubmix*, int32>& InChildCounts, bool bSelectNewNode = true);
 	/**
 	 * Recursively remove a node and its children from the graph
 	 */
-	void RecursivelyRemoveNode(class USoundSubmixGraphNode* ParentNode);
+	void RecursivelyRemoveNode(USoundSubmixGraphNode* ParentNode);
 	/**
 	 * Remove all Nodes from the graph
 	 */
@@ -105,14 +122,10 @@ private:
 	 * @param bSelectNewNode	Whether or not to select the new node being created
 	 * @return	Either a new node or an existing node representing the class
 	 */
-	class USoundSubmixGraphNode* CreateNode(class USoundSubmix* SoundSubmix, int32 NodePosX, int32 NodePosY, bool bSelectNewNode = true);
-	/**
-	 * Find an existing node that represents a given SoundSubmix
-	 */
-	class USoundSubmixGraphNode* FindExistingNode(class USoundSubmix* SoundSubmix) const;
+	USoundSubmixGraphNode* CreateNode(USoundSubmix* SoundSubmix, int32 NodePosX, int32 NodePosY, bool bSelectNewNode = true);
 
 private:
 	/** SoundSubmix which forms the root of this graph */
-	class USoundSubmix*	RootSoundSubmix;
+	USoundSubmix*	RootSoundSubmix;
 };
 
diff --git a/Engine/Source/Editor/AudioEditor/Classes/SoundSubmixGraph/SoundSubmixGraphNode.h b/Engine/Source/Editor/AudioEditor/Classes/SoundSubmixGraph/SoundSubmixGraphNode.h
--- a/Engine/Source/Editor/AudioEditor/Classes/SoundSubmixGraph/SoundSubmixGraphNode.h
+++ b/Engine/Source/Editor/AudioEditor/Classes/SoundSubmixGraph/SoundSubmixGraphNode.h
@@ -33,7 +33,6 @@ class USoundSubmixGraphNode : public UEdGraphNode
 	virtual void AutowireNewNode(UEdGraphPin* FromPin) override;
 	virtual bool CanCreateUnderSpecifiedSchema(const UEdGraphSchema* Schema) const override;
 	virtual FText GetNodeTitle(ENodeTitleType::Type TitleType) const override;
-	virtual bool CanUserDeleteNode() const override;
 	//~ End UEdGraphNode Interface.
 
 private:
diff --git a/Engine/Source/Editor/AudioEditor/Classes/SoundSubmixGraph/SoundSubmixGraphSchema.h b/Engine/Source/Editor/AudioEditor/Classes/SoundSubmixGraph/SoundSubmixGraphSchema.h
--- a/Engine/Source/Editor/AudioEditor/Classes/SoundSubmixGraph/SoundSubmixGraphSchema.h
+++ b/Engine/Source/Editor/AudioEditor/Classes/SoundSubmixGraph/SoundSubmixGraphSchema.h
@@ -1,9 +1,75 @@
 // Copyright 1998-2019 Epic Games, Inc. All Rights Reserved.
-
 #pragma once
-#include "EdGraph/EdGraphSchema.h"
+
+#include "AssetData.h"
+#include "ConnectionDrawingPolicy.h"
+#include "EdGraphUtilities.h"
+
 #include "SoundSubmixGraphSchema.generated.h"
 
+
+// Forward Declarations
+class UEdGraph;
+class UEdGraphSchema;
+class FMenuBuilder;
+
+struct FSoundSubmixGraphConnectionDrawingPolicyFactory : public FGraphPanelPinConnectionFactory
+{
+public:
+	virtual ~FSoundSubmixGraphConnectionDrawingPolicyFactory()
+	{
+	}
+
+	// FGraphPanelPinConnectionFactory
+	virtual class FConnectionDrawingPolicy* CreateConnectionPolicy(
+		const UEdGraphSchema* Schema,
+		int32 InBackLayerID,
+		int32 InFrontLayerID,
+		float ZoomFactor,
+		const FSlateRect& InClippingRect,
+		FSlateWindowElementList& InDrawElements,
+		UEdGraph* InGraphObj) const override;
+	// ~FGraphPanelPinConnectionFactory
+};
+
+
+// This class draws the connections for an UEdGraph using a SoundCue schema
+class FSoundSubmixGraphConnectionDrawingPolicy : public FConnectionDrawingPolicy
+{
+protected:
+	// Times for one execution pair within the current graph
+	struct FTimePair
+	{
+		double PredExecTime;
+		double ThisExecTime;
+
+		FTimePair()
+			: PredExecTime(0.0)
+			, ThisExecTime(0.0)
+		{
+		}
+	};
+
+	// Map of pairings
+	typedef TMap<UEdGraphNode*, FTimePair> FExecPairingMap;
+
+	// Map of nodes that preceded before a given node in the execution sequence (one entry for each pairing)
+	TMap<UEdGraphNode*, FExecPairingMap> PredecessorNodes;
+
+	UEdGraph* GraphObj;
+
+	float ActiveWireThickness;
+	float InactiveWireThickness;
+
+public:
+	FSoundSubmixGraphConnectionDrawingPolicy(int32 InBackLayerID, int32 InFrontLayerID, float ZoomFactor, const FSlateRect& InClippingRect, FSlateWindowElementList& InDrawElements, UEdGraph* InGraphObj);
+
+	// FConnectionDrawingPolicy interface
+	virtual void DetermineWiringStyle(UEdGraphPin* OutputPin, UEdGraphPin* InputPin, /*inout*/ FConnectionParams& Params) override;
+	// End of FConnectionDrawingPolicy interface
+};
+
+
 /** Action to add a node to the graph */
 USTRUCT()
 struct AUDIOEDITOR_API FSoundSubmixGraphSchemaAction_NewNode : public FEdGraphSchemaAction
@@ -25,7 +91,7 @@ struct AUDIOEDITOR_API FSoundSubmixGraphSchemaAction_NewNode : public FEdGraphSc
 
 	//~ Begin FEdGraphSchemaAction Interface
 	virtual FName GetTypeId() const override { return StaticGetTypeId(); } 
-	virtual UEdGraphNode* PerformAction(class UEdGraph* ParentGraph, UEdGraphPin* FromPin, const FVector2D Location, bool bSelectNewNode = true) override;
+	virtual UEdGraphNode* PerformAction(UEdGraph* ParentGraph, UEdGraphPin* FromPin, const FVector2D Location, bool bSelectNewNode = true) override;
 	//~ End FEdGraphSchemaAction Interface
 
 	/** Name for the new SoundSubmix */
@@ -54,7 +120,7 @@ class USoundSubmixGraphSchema : public UEdGraphSchema
 	virtual void BreakNodeLinks(UEdGraphNode& TargetNode) const override;
 	virtual void BreakPinLinks(UEdGraphPin& TargetPin, bool bSendsNodeNotifcation) const override;
 	virtual void BreakSinglePinLink(UEdGraphPin* SourcePin, UEdGraphPin* TargetPin) const override;
-	virtual void DroppedAssetsOnGraph(const TArray<struct FAssetData>& Assets, const FVector2D& GraphPosition, UEdGraph* Graph) const override;
+	virtual void DroppedAssetsOnGraph(const TArray<FAssetData>& Assets, const FVector2D& GraphPosition, UEdGraph* Graph) const override;
 	//~ End EdGraphSchema Interface
 };
 
diff --git a/Engine/Source/Editor/AudioEditor/Private/AudioEditorModule.cpp b/Engine/Source/Editor/AudioEditor/Private/AudioEditorModule.cpp
--- a/Engine/Source/Editor/AudioEditor/Private/AudioEditorModule.cpp
+++ b/Engine/Source/Editor/AudioEditor/Private/AudioEditorModule.cpp
@@ -39,6 +39,7 @@
 #include "Styling/SlateStyle.h"
 #include "Styling/SlateStyleRegistry.h"
 #include "SoundFileIO/SoundFileIO.h"
+#include "SoundSubmixGraph/SoundSubmixGraphSchema.h"
 
 const FName AudioEditorAppIdentifier = FName(TEXT("AudioEditorApp"));
 
@@ -79,10 +80,14 @@ public:
 		SoundSubmixExtensibility.Init();
 
 		// Register the sound cue graph connection policy with the graph editor
-		SoundCueGraphConnectionFactory = MakeShareable(new FSoundCueGraphConnectionDrawingPolicyFactory);
+		SoundCueGraphConnectionFactory = MakeShared<FSoundCueGraphConnectionDrawingPolicyFactory>();
 		FEdGraphUtilities::RegisterVisualPinConnectionFactory(SoundCueGraphConnectionFactory);
 
-		TSharedPtr<FSoundCueGraphNodeFactory> SoundCueGraphNodeFactory = MakeShareable(new FSoundCueGraphNodeFactory());
+		// Register the sound cue graph connection policy with the graph editor
+		SoundSubmixGraphConnectionFactory = MakeShared<FSoundSubmixGraphConnectionDrawingPolicyFactory>();
+		FEdGraphUtilities::RegisterVisualPinConnectionFactory(SoundSubmixGraphConnectionFactory);
+
+		TSharedPtr<FSoundCueGraphNodeFactory> SoundCueGraphNodeFactory = MakeShared<FSoundCueGraphNodeFactory>();
 		FEdGraphUtilities::RegisterVisualNodeFactory(SoundCueGraphNodeFactory);
 
 		// Create reimport handler for sound node waves
@@ -114,6 +119,11 @@ public:
 		{
 			FEdGraphUtilities::UnregisterVisualPinConnectionFactory(SoundCueGraphConnectionFactory);
 		}
+
+		if (SoundSubmixGraphConnectionFactory.IsValid())
+		{
+			FEdGraphUtilities::UnregisterVisualPinConnectionFactory(SoundSubmixGraphConnectionFactory);
+		}
 	}
 
 	virtual void RegisterAssetActions() override
@@ -189,9 +199,9 @@ public:
 
 	virtual TSharedRef<FAssetEditorToolkit> CreateSoundSubmixEditor(const EToolkitMode::Type Mode, const TSharedPtr< IToolkitHost >& InitToolkitHost, USoundSubmix* InSoundSubmix) override
 	{
-		TSharedRef<FSoundSubmixEditor> NewSoundSubmixEditor(new FSoundSubmixEditor());
-		NewSoundSubmixEditor->InitSoundSubmixEditor(Mode, InitToolkitHost, InSoundSubmix);
-		return NewSoundSubmixEditor;
+		TSharedPtr<FSoundSubmixEditor> NewSubmixEditor = MakeShared<FSoundSubmixEditor>();
+		NewSubmixEditor->Init(Mode, InitToolkitHost, InSoundSubmix);
+		return StaticCastSharedPtr<FAssetEditorToolkit>(NewSubmixEditor).ToSharedRef();
 	}
 
 	virtual TSharedPtr<FExtensibilityManager> GetSoundClassMenuExtensibilityManager() override
@@ -326,8 +336,8 @@ private:
 	FExtensibilityManagers SoundSubmixExtensibility;
 	TSet<USoundEffectPreset*> RegisteredActions;
 	TSharedPtr<FGraphPanelPinConnectionFactory> SoundCueGraphConnectionFactory;
+	TSharedPtr<FGraphPanelPinConnectionFactory> SoundSubmixGraphConnectionFactory;
 	TSharedPtr<FSlateStyleSet> AudioStyleSet;
-
 };
 
 IMPLEMENT_MODULE( FAudioEditorModule, AudioEditor );
diff --git a/Engine/Source/Editor/AudioEditor/Private/SoundSubmixEditor.cpp b/Engine/Source/Editor/AudioEditor/Private/SoundSubmixEditor.cpp
--- a/Engine/Source/Editor/AudioEditor/Private/SoundSubmixEditor.cpp
+++ b/Engine/Source/Editor/AudioEditor/Private/SoundSubmixEditor.cpp
@@ -2,59 +2,81 @@
 
 #include "SoundSubmixEditor.h"
 
-#include "Editor/PropertyEditor/Public/PropertyEditorModule.h"
-#include "Editor/PropertyEditor/Public/IDetailsView.h"
 #include "AssetRegistryModule.h"
-
-#include "Editor/WorkspaceMenuStructure/Public/WorkspaceMenuStructureModule.h"
-
+#include "AssetToolsModule.h"
+#include "AudioEditorModule.h"
+#include "Editor.h"
 #include "GraphEditor.h"
+#include "GraphEditAction.h"
 #include "Kismet2/BlueprintEditorUtils.h"
-#include "ScopedTransaction.h"
-#include "Toolkits/IToolkitHost.h"
-#include "Widgets/Docking/SDockTab.h"
+#include "Editor/PropertyEditor/Public/PropertyEditorModule.h"
+#include "Editor/PropertyEditor/Public/IDetailsView.h"
+#include "Editor/WorkspaceMenuStructure/Public/WorkspaceMenuStructureModule.h"
 #include "Framework/Commands/GenericCommands.h"
+#include "Framework/Docking/TabManager.h"
+#include "Factories/SoundSubmixFactory.h"
+#include "IAssetTools.h"
+#include "ScopedTransaction.h"
 #include "Sound/SoundSubmix.h"
 #include "SoundSubmixGraph/SoundSubmixGraph.h"
 #include "SoundSubmixGraph/SoundSubmixGraphNode.h"
 #include "SoundSubmixGraph/SoundSubmixGraphSchema.h"
-#include "AudioEditorModule.h"
 #include "SSoundSubmixActionMenu.h"
-#include "Editor.h"
-#include "AssetToolsModule.h"
-#include "IAssetTools.h"
-#include "Factories/SoundSubmixFactory.h"
+#include "Toolkits/IToolkitHost.h"
 #include "UObject/Package.h"
+#include "Widgets/Docking/SDockTab.h"
 
 #define LOCTEXT_NAMESPACE "SoundSubmixEditor"
-DEFINE_LOG_CATEGORY_STATIC( LogSoundSubmixEditor, Log, All );
+DEFINE_LOG_CATEGORY_STATIC(LogSoundSubmixEditor, Log, All);
 
-const FName FSoundSubmixEditor::GraphCanvasTabId( TEXT( "SoundSubmixEditor_GraphCanvas" ) );
-const FName FSoundSubmixEditor::PropertiesTabId( TEXT( "SoundSubmixEditor_Properties" ) );
+const FName FSoundSubmixEditor::GraphCanvasTabId(TEXT("SoundSubmixEditor_GraphCanvas"));
+const FName FSoundSubmixEditor::PropertiesTabId(TEXT("SoundSubmixEditor_Properties"));
 
 
-//////////////////////////////////////////////////////////////////////////
-// FSoundSubmixEditor
+class SSoundSubmixGraphEditor : public SGraphEditor
+{
+private:
+	FSoundSubmixEditor* SubmixEditor;
 
-void FSoundSubmixEditor::RegisterTabSpawners(const TSharedRef<class FTabManager>& InTabManager)
+public:
+	SSoundSubmixGraphEditor()
+		: SGraphEditor()
+		, SubmixEditor(nullptr)
+	{
+	}
+
+	void Construct(const FArguments& InArgs, FSoundSubmixEditor& InEditor)
+	{
+		SubmixEditor = &InEditor;
+		SGraphEditor::Construct(InArgs);
+	}
+
+	void OnGraphChanged(const FEdGraphEditAction& InAction) override
+	{
+		SubmixEditor->AddMissingEditableSubmixes(InAction.Graph);
+		SGraphEditor::OnGraphChanged(InAction);
+	}
+};
+
+void FSoundSubmixEditor::RegisterTabSpawners(const TSharedRef<FTabManager>& InTabManager)
 {
 	WorkspaceMenuCategory = InTabManager->AddLocalWorkspaceMenuCategory(LOCTEXT("WorkspaceMenu_SoundSubmixEditor", "Sound Submix Editor"));
 	auto WorkspaceMenuCategoryRef = WorkspaceMenuCategory.ToSharedRef();
 
 	FAssetEditorToolkit::RegisterTabSpawners(InTabManager);
 
-	InTabManager->RegisterTabSpawner( GraphCanvasTabId, FOnSpawnTab::CreateSP(this, &FSoundSubmixEditor::SpawnTab_GraphCanvas) )
-		.SetDisplayName( LOCTEXT( "GraphCanvasTab", "Graph" ) )
-		.SetGroup( WorkspaceMenuCategoryRef )
+	InTabManager->RegisterTabSpawner(GraphCanvasTabId, FOnSpawnTab::CreateSP(this, &FSoundSubmixEditor::SpawnTab_GraphCanvas))
+		.SetDisplayName(LOCTEXT("GraphCanvasTab", "Graph"))
+		.SetGroup(WorkspaceMenuCategoryRef)
 		.SetIcon(FSlateIcon(FEditorStyle::GetStyleSetName(), "GraphEditor.EventGraph_16x"));
 
-	InTabManager->RegisterTabSpawner( PropertiesTabId, FOnSpawnTab::CreateSP(this, &FSoundSubmixEditor::SpawnTab_Properties) )
-		.SetDisplayName( LOCTEXT( "PropertiesTab", "Details" ) )
-		.SetGroup( WorkspaceMenuCategoryRef )
+	InTabManager->RegisterTabSpawner(PropertiesTabId, FOnSpawnTab::CreateSP(this, &FSoundSubmixEditor::SpawnTab_Properties))
+		.SetDisplayName(LOCTEXT("PropertiesTab", "Details"))
+		.SetGroup(WorkspaceMenuCategoryRef)
 		.SetIcon(FSlateIcon(FEditorStyle::GetStyleSetName(), "LevelEditor.Tabs.Details"));
 }
 
-void FSoundSubmixEditor::UnregisterTabSpawners(const TSharedRef<class FTabManager>& InTabManager)
+void FSoundSubmixEditor::UnregisterTabSpawners(const TSharedRef<FTabManager>& InTabManager)
 {
 	FAssetEditorToolkit::UnregisterTabSpawners(InTabManager);
 
@@ -62,74 +84,92 @@ void FSoundSubmixEditor::UnregisterTabSpawners(const TSharedRef<class FTabManage
 	InTabManager->UnregisterTabSpawner( PropertiesTabId );
 }
 
-void FSoundSubmixEditor::InitSoundSubmixEditor( const EToolkitMode::Type Mode, const TSharedPtr< class IToolkitHost >& InitToolkitHost, UObject* ObjectToEdit )
+void FSoundSubmixEditor::AddEditableSubmixChildren(USoundSubmix* RootSubmix)
+{
+	if (!RootSubmix)
+	{
+		return;
+	}
+
+	RootSubmix->SetFlags(RF_Transactional);
+	for (USoundSubmix* Child : RootSubmix->ChildSubmixes)
+	{
+		if (Child)
+		{
+			Child->SoundSubmixGraph = RootSubmix->SoundSubmixGraph;
+			AddEditingObject(Child);
+			AddEditableSubmixChildren(Child);
+		}
+	}
+}
+
+void FSoundSubmixEditor::Init(const EToolkitMode::Type Mode, const TSharedPtr<IToolkitHost >& InitToolkitHost, UObject* ObjectToEdit)
 {
-	SoundSubmix = CastChecked<USoundSubmix>(ObjectToEdit);
+	USoundSubmix* SoundSubmix = CastChecked<USoundSubmix>(ObjectToEdit);
 
 	while (SoundSubmix->ParentSubmix)
 	{
 		SoundSubmix = SoundSubmix->ParentSubmix;
 	}
 
-	// Support undo/redo
-	SoundSubmix->SetFlags(RF_Transactional);
-
 	GEditor->RegisterForUndo(this);
 
 	ToolkitCommands->MapAction(
 		FGenericCommands::Get().Undo,
-		FExecuteAction::CreateSP( this, &FSoundSubmixEditor::UndoGraphAction ));
+		FExecuteAction::CreateSP(this, &FSoundSubmixEditor::UndoGraphAction));
 
 	ToolkitCommands->MapAction(
 		FGenericCommands::Get().Redo,
-		FExecuteAction::CreateSP( this, &FSoundSubmixEditor::RedoGraphAction ));
+		FExecuteAction::CreateSP(this, &FSoundSubmixEditor::RedoGraphAction));
 
-	if( !SoundSubmix->SoundSubmixGraph )
+	if (!SoundSubmix->SoundSubmixGraph)
 	{
 		USoundSubmixGraph* SoundSubmixGraph = CastChecked<USoundSubmixGraph>(FBlueprintEditorUtils::CreateNewGraph(SoundSubmix, NAME_None, USoundSubmixGraph::StaticClass(), USoundSubmixGraphSchema::StaticClass()));
 		SoundSubmixGraph->SetRootSoundSubmix(SoundSubmix);
 
 		SoundSubmix->SoundSubmixGraph = SoundSubmixGraph;
+		SoundSubmixGraph->RebuildGraph();
 	}
 
-	CastChecked<USoundSubmixGraph>(SoundSubmix->SoundSubmixGraph)->RebuildGraph();
-
-	CreateInternalWidgets();
+	CreateInternalWidgets(SoundSubmix);
 
-	TSharedRef<FTabManager::FLayout> StandaloneDefaultLayout = FTabManager::NewLayout( "Standalone_SoundSubmixEditor_Layout_v2" )
-	->AddArea
-	(
-		FTabManager::NewPrimaryArea() ->SetOrientation(Orient_Vertical)
-		->Split
-		(
-			FTabManager::NewStack()
-			->SetSizeCoefficient(0.1f)
-			->SetHideTabWell( true )
-			->AddTab(GetToolbarTabId(), ETabState::OpenedTab)
-		)
-		->Split
+	TSharedRef<FTabManager::FLayout> StandaloneDefaultLayout = FTabManager::NewLayout("Standalone_SoundSubmixEditor_Layout_v2")
+		->AddArea
 		(
-			FTabManager::NewSplitter()
+			FTabManager::NewPrimaryArea()->SetOrientation(Orient_Vertical)
 			->Split
 			(
 				FTabManager::NewStack()
-				->SetSizeCoefficient(0.2f)
-				->AddTab(PropertiesTabId, ETabState::OpenedTab)
+				->SetSizeCoefficient(0.1f)
+				->SetHideTabWell(true)
+				->AddTab(GetToolbarTabId(), ETabState::OpenedTab)
 			)
 			->Split
 			(
-				FTabManager::NewStack()
-				->SetSizeCoefficient(0.8f)
-				->AddTab(GraphCanvasTabId, ETabState::OpenedTab)
+				FTabManager::NewSplitter()
+				->Split
+				(
+					FTabManager::NewStack()
+					->SetSizeCoefficient(0.2f)
+					->AddTab(PropertiesTabId, ETabState::OpenedTab)
+				)
+				->Split
+				(
+					FTabManager::NewStack()
+					->SetSizeCoefficient(0.8f)
+					->AddTab(GraphCanvasTabId, ETabState::OpenedTab)
+				)
 			)
-		)
-	);
+		);
 
 	const bool bCreateDefaultStandaloneMenu = true;
 	const bool bCreateDefaultToolbar = true;
-	FAssetEditorToolkit::InitAssetEditor( Mode, InitToolkitHost, TEXT("SoundSubmixEditorApp"), StandaloneDefaultLayout, bCreateDefaultStandaloneMenu, bCreateDefaultToolbar, SoundSubmix);
+	FAssetEditorToolkit::InitAssetEditor(Mode, InitToolkitHost, TEXT("SoundSubmixEditorApp"), StandaloneDefaultLayout, bCreateDefaultStandaloneMenu, bCreateDefaultToolbar, SoundSubmix);
+
+	// Must be called after super class initialization
+	AddEditableSubmixChildren(SoundSubmix);
 
-	IAudioEditorModule* AudioEditorModule = &FModuleManager::LoadModuleChecked<IAudioEditorModule>( "AudioEditor" );
+	IAudioEditorModule* AudioEditorModule = &FModuleManager::LoadModuleChecked<IAudioEditorModule>("AudioEditor");
 	AddMenuExtender(AudioEditorModule->GetSoundSubmixMenuExtensibilityManager()->GetAllExtenders(GetToolkitCommands(), GetEditingObjects()));
 	AddToolbarExtender(AudioEditorModule->GetSoundSubmixToolBarExtensibilityManager()->GetAllExtenders(GetToolkitCommands(), GetEditingObjects()));
 
@@ -147,20 +187,18 @@ void FSoundSubmixEditor::InitSoundSubmixEditor( const EToolkitMode::Type Mode, c
 	}
 }
 
-FSoundSubmixEditor::FSoundSubmixEditor()
-	: SoundSubmix(nullptr)
-{
-}
-
 FSoundSubmixEditor::~FSoundSubmixEditor()
 {
-	GEditor->UnregisterForUndo( this );
+	GEditor->UnregisterForUndo(this);
 	DetailsView.Reset();
 }
 
-void FSoundSubmixEditor::AddReferencedObjects( FReferenceCollector& Collector )
+void FSoundSubmixEditor::AddReferencedObjects(FReferenceCollector& Collector)
 {
-	Collector.AddReferencedObject(SoundSubmix);
+	for (UObject* Obj : GetEditingObjects())
+	{
+		Collector.AddReferencedObject(Obj);
+	}
 }
 
 TSharedRef<SDockTab> FSoundSubmixEditor::SpawnTab_GraphCanvas(const FSpawnTabArgs& Args)
@@ -200,6 +238,11 @@ FText FSoundSubmixEditor::GetBaseToolkitName() const
 	return LOCTEXT("AppLabel", "Sound Submix Editor");
 }
 
+FText FSoundSubmixEditor::GetToolkitToolTipText() const
+{
+	return GetToolTipTextForObject(GetEditingObjects()[0]);
+}
+
 FString FSoundSubmixEditor::GetWorldCentricTabPrefix() const
 {
 	return LOCTEXT( "WorldCentricTabPrefix", "Sound Submix " ).ToString();
@@ -210,32 +253,32 @@ FLinearColor FSoundSubmixEditor::GetWorldCentricTabColorScale() const
 	return FLinearColor( 0.2f, 0.4f, 0.8f, 0.5f );
 }
 
-void FSoundSubmixEditor::CreateInternalWidgets()
+void FSoundSubmixEditor::CreateInternalWidgets(USoundSubmix* InSoundSubmix)
 {
-	GraphEditor = CreateGraphEditorWidget();
+	GraphEditor = CreateGraphEditorWidget(InSoundSubmix);
 
-	FPropertyEditorModule& PropertyEditorModule = FModuleManager::GetModuleChecked<FPropertyEditorModule>( "PropertyEditor" );
-	const FDetailsViewArgs DetailsViewArgs( false, false, true, FDetailsViewArgs::ObjectsUseNameArea, false );
-	DetailsView = PropertyEditorModule.CreateDetailView( DetailsViewArgs );
-	DetailsView->SetObject( SoundSubmix );
+	FPropertyEditorModule& PropertyEditorModule = FModuleManager::GetModuleChecked<FPropertyEditorModule>("PropertyEditor");
+	const FDetailsViewArgs DetailsViewArgs(false, false, true, FDetailsViewArgs::ObjectsUseNameArea, false);
+	DetailsView = PropertyEditorModule.CreateDetailView(DetailsViewArgs);
+	DetailsView->SetObject(InSoundSubmix);
 }
 
-TSharedRef<SGraphEditor> FSoundSubmixEditor::CreateGraphEditorWidget()
+TSharedRef<SGraphEditor> FSoundSubmixEditor::CreateGraphEditorWidget(USoundSubmix* InSoundSubmix)
 {
-	if ( !GraphEditorCommands.IsValid() )
+	if (!GraphEditorCommands.IsValid())
 	{
-		GraphEditorCommands = MakeShareable( new FUICommandList );
+		GraphEditorCommands = MakeShareable(new FUICommandList);
 
 		// Editing commands
-		GraphEditorCommands->MapAction( FGenericCommands::Get().SelectAll,
-			FExecuteAction::CreateSP( this, &FSoundSubmixEditor::SelectAllNodes ),
-			FCanExecuteAction::CreateSP( this, &FSoundSubmixEditor::CanSelectAllNodes )
-			);
-
-		GraphEditorCommands->MapAction( FGenericCommands::Get().Delete,
-			FExecuteAction::CreateSP( this, &FSoundSubmixEditor::RemoveSelectedNodes ),
-			FCanExecuteAction::CreateSP( this, &FSoundSubmixEditor::CanRemoveNodes )
-			);
+		GraphEditorCommands->MapAction(FGenericCommands::Get().SelectAll,
+			FExecuteAction::CreateSP(this, &FSoundSubmixEditor::SelectAllNodes),
+			FCanExecuteAction::CreateSP(this, &FSoundSubmixEditor::CanSelectAllNodes)
+		);
+
+		GraphEditorCommands->MapAction(FGenericCommands::Get().Delete,
+			FExecuteAction::CreateSP(this, &FSoundSubmixEditor::RemoveSelectedNodes),
+			FCanExecuteAction::CreateSP(this, &FSoundSubmixEditor::CanRemoveNodes)
+		);
 	}
 
 	FGraphAppearanceInfo AppearanceInfo;
@@ -245,31 +288,30 @@ TSharedRef<SGraphEditor> FSoundSubmixEditor::CreateGraphEditorWidget()
 	InEvents.OnSelectionChanged = SGraphEditor::FOnSelectionChanged::CreateSP(this, &FSoundSubmixEditor::OnSelectedNodesChanged);
 	InEvents.OnCreateActionMenu = SGraphEditor::FOnCreateActionMenu::CreateSP(this, &FSoundSubmixEditor::OnCreateGraphActionMenu);
 
-	return SNew(SGraphEditor)
+	return SNew(SSoundSubmixGraphEditor, *this)
 		.AdditionalCommands(GraphEditorCommands)
 		.IsEditable(true)
 		.Appearance(AppearanceInfo)
-		.GraphToEdit(SoundSubmix->SoundSubmixGraph)
+		.GraphToEdit(InSoundSubmix->SoundSubmixGraph)
 		.GraphEvents(InEvents)
 		.ShowGraphStateOverlay(false);
 }
 
-void FSoundSubmixEditor::OnSelectedNodesChanged(const TSet<class UObject*>& NewSelection)
+void FSoundSubmixEditor::OnSelectedNodesChanged(const TSet<UObject*>& NewSelection)
 {
-	TArray<UObject*> Selection;
-
-	if(NewSelection.Num())
+	if(NewSelection.Num() > 0)
 	{
-		for(TSet<class UObject*>::TConstIterator SetIt(NewSelection);SetIt;++SetIt)
+		TArray<UObject*> Selection;
+		for (UObject* Obj : NewSelection)
 		{
-			USoundSubmixGraphNode* GraphNode = CastChecked<USoundSubmixGraphNode>(*SetIt);
+			USoundSubmixGraphNode* GraphNode = CastChecked<USoundSubmixGraphNode>(Obj);
 			Selection.Add(GraphNode->SoundSubmix);
 		}
 		DetailsView->SetObjects(Selection);
 	}
 	else
 	{
-		DetailsView->SetObject(SoundSubmix);
+		DetailsView->SetObject(GetEditingObjects()[0]);
 	}
 }
 
@@ -291,6 +333,33 @@ void FSoundSubmixEditor::SelectAllNodes()
 	GraphEditor->SelectAllNodes();
 }
 
+void FSoundSubmixEditor::SelectSubmixes(TSet<USoundSubmix*>& InSubmixes)
+{
+	TArray<UObject*> ObjectsToSelect;
+
+	GraphEditor->SelectAllNodes();
+	const TSet<UObject*> SelectedNodes = GraphEditor->GetSelectedNodes();
+	GraphEditor->ClearSelectionSet();
+
+	for (UObject* SelectedNode : SelectedNodes)
+	{
+		if (SelectedNode)
+		{
+			USoundSubmixGraphNode* GraphNode = CastChecked<USoundSubmixGraphNode>(SelectedNode);
+			if (USoundSubmix* Submix = GraphNode->SoundSubmix)
+			{
+				if (InSubmixes.Contains(Submix))
+				{
+					ObjectsToSelect.Add(Submix);
+					GraphEditor->SetNodeSelection(GraphNode, true /* bSelect */);
+				}
+			}
+		}
+	}
+
+	DetailsView->SetObjects(ObjectsToSelect);
+}
+
 bool FSoundSubmixEditor::CanSelectAllNodes() const
 {
 	return true;
@@ -298,28 +367,34 @@ bool FSoundSubmixEditor::CanSelectAllNodes() const
 
 void FSoundSubmixEditor::RemoveSelectedNodes()
 {
-	const FScopedTransaction Transaction( LOCTEXT("SoundSubmixEditorRemoveSelectedNode", "Sound Submix Editor: Remove Selected SoundSubmixes from editor") );
+	const FScopedTransaction Transaction(LOCTEXT("SoundSubmixEditorRemoveSelectedNode", "Sound Submix Editor: Remove Selected SoundSubmixes from editor"));
 
-	CastChecked<USoundSubmixGraph>(SoundSubmix->SoundSubmixGraph)->RecursivelyRemoveNodes(GraphEditor->GetSelectedNodes());
-
-	GraphEditor->ClearSelectionSet();
-}
+	int32 NumObjectsRemoved = GetEditingObjects().Num();
 
-bool FSoundSubmixEditor::CanRemoveNodes() const
-{
-	const FGraphPanelSelectionSet SelectedNodes = GraphEditor->GetSelectedNodes();
-
-	for (FGraphPanelSelectionSet::TConstIterator NodeIt(SelectedNodes); NodeIt; ++NodeIt)
+	const TSet<UObject*> SelectedNodes = GraphEditor->GetSelectedNodes();
+	for (UObject* SelectedNode : SelectedNodes)
 	{
-		USoundSubmixGraphNode* Node = Cast<USoundSubmixGraphNode>(*NodeIt);
-
-		if (Node && Node->CanUserDeleteNode())
+		USoundSubmixGraphNode* Node = Cast<USoundSubmixGraphNode>(SelectedNode);
+		if (Node && Node->SoundSubmix && Node->CanUserDeleteNode())
 		{
-			return true;
+			NumObjectsRemoved--;
+			RemoveEditingObject(Node->SoundSubmix);
 		}
 	}
 
-	return false;
+	check(NumObjectsRemoved >= 0);
+	if (NumObjectsRemoved > 0)
+	{
+		USoundSubmix* SoundSubmix = CastChecked<USoundSubmix>(GetEditingObjects()[0]);
+		USoundSubmixGraph* Graph = CastChecked<USoundSubmixGraph>(SoundSubmix->SoundSubmixGraph);
+		Graph->RecursivelyRemoveNodes(SelectedNodes);
+		GraphEditor->ClearSelectionSet();
+	}
+}
+
+bool FSoundSubmixEditor::CanRemoveNodes() const
+{
+	return GraphEditor->GetSelectedNodes().Num() > 0;
 }
 
 void FSoundSubmixEditor::UndoGraphAction()
@@ -337,29 +412,72 @@ void FSoundSubmixEditor::RedoGraphAction()
 
 void FSoundSubmixEditor::CreateSoundSubmix(UEdGraphPin* FromPin, const FVector2D Location, const FString& Name)
 {
-	// If we have a valid name
-	if (!Name.IsEmpty() && Name != SoundSubmix->GetName())
+	if (Name.IsEmpty())
 	{
-		// Derive new package path from existing asset's path
-		FString PackagePath = SoundSubmix->GetPathName();
-		FString AssetName = FString::Printf(TEXT("/%s.%s"), *SoundSubmix->GetName(), *SoundSubmix->GetName());
-		PackagePath.RemoveFromEnd(AssetName);
+		return;
+	}
 
-		// Create a sound submix factory to create a new sound class
-		USoundSubmixFactory* SoundSubmixFactory = NewObject<USoundSubmixFactory>();
+	for (UObject* Obj : GetEditingObjects())
+	{
+		if (Obj->GetName() == Name)
+		return;
+	}
 
-		// Load asset tools to create the asset properly
-		FAssetToolsModule& AssetToolsModule = FModuleManager::LoadModuleChecked<FAssetToolsModule>("AssetTools");
-		USoundSubmix* NewSoundSubmix = Cast<USoundSubmix>(AssetToolsModule.Get().CreateAsset(Name, PackagePath, USoundSubmix::StaticClass(), SoundSubmixFactory, FName("SoundSubmixEditorNewAsset")));
+	// Derive new package path from existing asset's path
+	USoundSubmix* SoundSubmix = CastChecked<USoundSubmix>(GetEditingObjects()[0]);
+	FString PackagePath = SoundSubmix->GetPathName();
+	FString AssetName = FString::Printf(TEXT("/%s.%s"), *SoundSubmix->GetName(), *SoundSubmix->GetName());
+	PackagePath.RemoveFromEnd(AssetName);
 
-		if (NewSoundSubmix)
-		{
-			CastChecked<USoundSubmixGraph>(SoundSubmix->SoundSubmixGraph)->AddNewSoundSubmix(FromPin, NewSoundSubmix, Location.X, Location.Y);
+	// Create a sound submix factory to create a new sound class
+	USoundSubmixFactory* SoundSubmixFactory = NewObject<USoundSubmixFactory>();
+
+	// Load asset tools to create the asset properly
+	FAssetToolsModule& AssetToolsModule = FModuleManager::LoadModuleChecked<FAssetToolsModule>("AssetTools");
+	USoundSubmix* NewSoundSubmix = Cast<USoundSubmix>(AssetToolsModule.Get().CreateAsset(Name, PackagePath, USoundSubmix::StaticClass(), SoundSubmixFactory, FName("SoundSubmixEditorNewAsset")));
+
+	if (NewSoundSubmix)
+	{
+		CastChecked<USoundSubmixGraph>(SoundSubmix->SoundSubmixGraph)->AddNewSoundSubmix(FromPin, NewSoundSubmix, Location.X, Location.Y);
+
+		AddEditingObject(NewSoundSubmix);
+		NewSoundSubmix->PostEditChange();
+		NewSoundSubmix->MarkPackageDirty();
+	}
+}
+
+FText FSoundSubmixEditor::GetToolkitName() const
+{
+	UObject* EditObject = GetEditingObjects()[0];
+	return GetLabelForObject(EditObject);
+}
 
-			NewSoundSubmix->PostEditChange();
-			NewSoundSubmix->MarkPackageDirty();
+void FSoundSubmixEditor::AddMissingEditableSubmixes(UEdGraph* InGraph)
+{
+	UEdGraph* Graph = GraphEditor->GetCurrentGraph();
+	check(Graph);
+
+	if (InGraph != Graph)
+	{
+		return;
+	}
+
+	bool bRefreshGraphs = false;
+	if (Graph->Nodes.Num() > GetEditingObjects().Num())
+	{
+		for (UEdGraphNode* Node : Graph->Nodes)
+		{
+			USoundSubmixGraphNode* GraphNode = CastChecked<USoundSubmixGraphNode>(Node);
+			USoundSubmix* UntrackedSubmix = GraphNode->SoundSubmix;
+			if (UntrackedSubmix && !GetEditingObjects().Contains(UntrackedSubmix))
+			{
+				AddEditingObject(UntrackedSubmix);
+				bRefreshGraphs = true;
+			}
 		}
 	}
+
+	GraphEditor->NotifyGraphChanged();
 }
 
 void FSoundSubmixEditor::PostUndo(bool bSuccess)
diff --git a/Engine/Source/Editor/AudioEditor/Private/SoundSubmixEditor.h b/Engine/Source/Editor/AudioEditor/Private/SoundSubmixEditor.h
--- a/Engine/Source/Editor/AudioEditor/Private/SoundSubmixEditor.h
+++ b/Engine/Source/Editor/AudioEditor/Private/SoundSubmixEditor.h
@@ -1,19 +1,21 @@
 // Copyright 1998-2019 Epic Games, Inc. All Rights Reserved.
-
 #pragma once
 
+#include "AudioEditorModule.h"
 #include "GraphEditor.h"
-#include "ISoundSubmixEditor.h"
 #include "EditorUndoClient.h"
+#include "Framework/Docking/TabManager.h"
+#include "ISoundSubmixEditor.h"
+#include "Toolkits/IToolkitHost.h"
+#include "UObject/Object.h"
+#include "Widgets/Docking/SDockableTab.h"
 
-//////////////////////////////////////////////////////////////////////////
-// FSoundSubmixEditor
 
+// Forward Declarations
+class UEdGraphPin;
 class USoundSubmix;
 
-class FSoundSubmixEditor :	public ISoundSubmixEditor, 
-							public FGCObject, 
-							public FEditorUndoClient
+class FSoundSubmixEditor : public ISoundSubmixEditor, public FGCObject, public FEditorUndoClient
 {
 public:
 	SLATE_BEGIN_ARGS(FSoundSubmixEditor)
@@ -22,8 +24,8 @@ public:
 
 	SLATE_END_ARGS()
 
-	virtual void RegisterTabSpawners(const TSharedRef<class FTabManager>& TabManager) override;
-	virtual void UnregisterTabSpawners(const TSharedRef<class FTabManager>& TabManager) override;
+	virtual void RegisterTabSpawners(const TSharedRef<FTabManager>& TabManager) override;
+	virtual void UnregisterTabSpawners(const TSharedRef<FTabManager>& TabManager) override;
 
 	/**
 	 * Edits the specified sound submix object
@@ -32,47 +34,56 @@ public:
 	 * @param	InitToolkitHost			When Mode is WorldCentric, this is the level editor instance to spawn this editor within
 	 * @param	ObjectToEdit			The sound submix to edit
 	 */
-	void InitSoundSubmixEditor( const EToolkitMode::Type Mode, const TSharedPtr< class IToolkitHost >& InitToolkitHost, UObject* ObjectToEdit );
+	void Init(const EToolkitMode::Type Mode, const TSharedPtr<IToolkitHost>& InitToolkitHost, UObject* ObjectToEdit);
 
-	FSoundSubmixEditor();
 	virtual ~FSoundSubmixEditor();
 
 	/** FGCObject interface */
-	virtual void AddReferencedObjects( FReferenceCollector& Collector ) override;
+	virtual void AddReferencedObjects(FReferenceCollector& Collector) override;
 
-	/** IToolkit interface */
-	virtual FName GetToolkitFName() const override;
+	/** FAssetEditorToolkit interface */
 	virtual FText GetBaseToolkitName() const override;
+	virtual FName GetToolkitFName() const override;
+	virtual FText GetToolkitName() const override;
+	virtual FText GetToolkitToolTipText() const override;
 	virtual FString GetWorldCentricTabPrefix() const override;
 
 	/** @return Returns the color and opacity to use for the color that appears behind the tab text for this toolkit's tab in world-centric mode. */
 	virtual FLinearColor GetWorldCentricTabColorScale() const override;
 
 	//~ Begin ISoundSubmixEditor
-	void CreateSoundSubmix(class UEdGraphPin* FromPin, FVector2D Location, const FString& Name) override;
+	void CreateSoundSubmix(UEdGraphPin* FromPin, FVector2D Location, const FString& Name) override;
 	//~ End ISoundSubmixEditor
 
 	/** FEditorUndoClient Interface */
 	virtual void PostUndo(bool bSuccess) override;
 	virtual void PostRedo(bool bSuccess) override { PostUndo(bSuccess); }
 
+	void AddMissingEditableSubmixes(UEdGraph* InGraph);
+
+	/** Select node associated with the provided submix */
+	void SelectSubmixes(TSet<USoundSubmix*>& InSubmixes);
+
 private:
 	TSharedRef<SDockTab> SpawnTab_GraphCanvas(const FSpawnTabArgs& Args);
 	TSharedRef<SDockTab> SpawnTab_Properties(const FSpawnTabArgs& Args);
 
 private:
 	/** Creates all internal widgets for the tabs to point at */
-	void CreateInternalWidgets();
+	void CreateInternalWidgets(USoundSubmix* InSoundSubmix);
 
 	/** Create new graph editor widget */
-	TSharedRef<SGraphEditor> CreateGraphEditorWidget();
+	TSharedRef<SGraphEditor> CreateGraphEditorWidget(USoundSubmix* InSoundSubmix);
 
 	/** Called when the selection changes in the GraphEditor */
-	void OnSelectedNodesChanged(const TSet<class UObject*>& NewSelection);
+	void OnSelectedNodesChanged(const TSet<UObject*>& NewSelection);
 
 	/** Called to create context menu when right-clicking on graph */
 	FActionMenuContent OnCreateGraphActionMenu(UEdGraph* InGraph, const FVector2D& InNodePosition, const TArray<UEdGraphPin*>& InDraggedPins, bool bAutoExpand, SGraphEditor::FActionMenuClosed InOnMenuClosed);
 
+	/** Adds all children of provided root submix as editable */
+	void AddEditableSubmixChildren(USoundSubmix* RootSubmix);
+
 	/** Select every node in the graph */
 	void SelectAllNodes();
 	
@@ -92,17 +103,14 @@ private:
 	void RedoGraphAction();
 
 private:
-	/** The SoundSubmix asset being inspected */
-	USoundSubmix* SoundSubmix;
-
 	/** List of open tool panels; used to ensure only one exists at any one time */
-	TMap< FName, TWeakPtr<class SDockableTab> > SpawnedToolPanels;
+	TMap<FName, TWeakPtr<SDockableTab>> SpawnedToolPanels;
 
 	/** Graph Editor */
 	TSharedPtr<SGraphEditor> GraphEditor;
 
 	/** Property View */
-	TSharedPtr<class IDetailsView> DetailsView;
+	TSharedPtr<IDetailsView> DetailsView;
 
 	/** Command list for this editor */
 	TSharedPtr<FUICommandList> GraphEditorCommands;
diff --git a/Engine/Source/Editor/AudioEditor/Private/SoundSubmixGraph.cpp b/Engine/Source/Editor/AudioEditor/Private/SoundSubmixGraph.cpp
--- a/Engine/Source/Editor/AudioEditor/Private/SoundSubmixGraph.cpp
+++ b/Engine/Source/Editor/AudioEditor/Private/SoundSubmixGraph.cpp
@@ -1,31 +1,16 @@
 // Copyright 1998-2019 Epic Games, Inc. All Rights Reserved.
-
-/////////////////////////////////////////////////////
-// USoundSubmixGraph
-
 #include "SoundSubmixGraph/SoundSubmixGraph.h"
+
+#include "GraphEditor.h"
 #include "Sound/SoundSubmix.h"
 #include "SoundSubmixGraph/SoundSubmixGraphNode.h"
-#include "GraphEditor.h"
 #include "UObject/Package.h"
 
-class FSoundSubmixAudioEditor : public ISoundSubmixAudioEditor
-{
-public:
-	void RefreshGraphLinks(UEdGraph* SoundSubmixGraph) override
-	{
-		CastChecked<USoundSubmixGraph>(SoundSubmixGraph)->RefreshGraphLinks();
-	}
-};
 
 USoundSubmixGraph::USoundSubmixGraph(const FObjectInitializer& ObjectInitializer)
 	: Super(ObjectInitializer)
-	, RootSoundSubmix(NULL)
+	, RootSoundSubmix(nullptr)
 {
-	if (!USoundSubmix::GetSoundSubmixAudioEditor().IsValid())
-	{
-		USoundSubmix::SetSoundSubmixAudioEditor(TSharedPtr<ISoundSubmixAudioEditor>(new FSoundSubmixAudioEditor()));
-	}
 }
 
 void USoundSubmixGraph::SetRootSoundSubmix(USoundSubmix* InSoundSubmix)
@@ -70,7 +55,7 @@ void USoundSubmixGraph::AddDroppedSoundSubmixes(const TArray<USoundSubmix*>& Sou
 	NotifyGraphChanged();
 }
 
-void USoundSubmixGraph::AddNewSoundSubmix(UEdGraphPin* FromPin, class USoundSubmix* SoundSubmix, int32 NodePosX, int32 NodePosY, bool bSelectNewNode/* = true*/)
+void USoundSubmixGraph::AddNewSoundSubmix(UEdGraphPin* FromPin, USoundSubmix* SoundSubmix, int32 NodePosX, int32 NodePosY, bool bSelectNewNode/* = true*/)
 {
 	check(SoundSubmix->ChildSubmixes.Num() == 0);
 
@@ -124,8 +109,6 @@ void USoundSubmixGraph::LinkSoundSubmixes()
 			Node->SoundSubmix->MarkPackageDirty();
 		}
 	}
-
-	RootSoundSubmix->RefreshAllGraphs(true);
 }
 
 void USoundSubmixGraph::RefreshGraphLinks()
@@ -173,7 +156,7 @@ void USoundSubmixGraph::RefreshGraphLinks()
 	NotifyGraphChanged();
 }
 
-void USoundSubmixGraph::RecursivelyRemoveNodes(const TSet<class UObject*> NodesToRemove)
+void USoundSubmixGraph::RecursivelyRemoveNodes(const TSet<UObject*> NodesToRemove)
 {
 	Modify();
 
@@ -190,7 +173,7 @@ void USoundSubmixGraph::RecursivelyRemoveNodes(const TSet<class UObject*> NodesT
 	LinkSoundSubmixes();
 }
 
-int32 USoundSubmixGraph::ConstructNodes(class USoundSubmix* SoundSubmix, int32 NodePosX, int32 NodePosY, bool bSelectNewNode/* = true*/)
+int32 USoundSubmixGraph::ConstructNodes(USoundSubmix* SoundSubmix, int32 NodePosX, int32 NodePosY, bool bSelectNewNode/* = true*/)
 {
 	check(SoundSubmix);
 
@@ -224,10 +207,10 @@ int32 USoundSubmixGraph::RecursivelyGatherChildCounts(USoundSubmix* ParentSubmix
 	return ChildSize;
 }
 
-int32 USoundSubmixGraph::RecursivelyConstructChildNodes(USoundSubmixGraphNode* ParentNode, const TMap<USoundSubmix*, int32>& InChildCounts, bool bSelectNewNode/* = true*/)
+int32 USoundSubmixGraph::RecursivelyConstructChildNodes(USoundSubmixGraphNode* ParentNode, const TMap<USoundSubmix*, int32>& InChildCounts, bool bSelectNewNode /* = true*/)
 {
-	const int32 HorizontalSpacing = 400;
-	const int32 VerticalSpacing = 100;
+	static const int32 HorizontalSpacing = -400;
+	static const int32 VerticalSpacing = 100;
 
 	USoundSubmix* ParentSubmix = ParentNode->SoundSubmix;
 	int32 TotalChildSizeY = InChildCounts.FindChecked(ParentSubmix) * VerticalSpacing;
@@ -250,7 +233,7 @@ int32 USoundSubmixGraph::RecursivelyConstructChildNodes(USoundSubmixGraphNode* P
 	return TotalChildSizeY;
 }
 
-void USoundSubmixGraph::RecursivelyRemoveNode(class USoundSubmixGraphNode* ParentNode)
+void USoundSubmixGraph::RecursivelyRemoveNode(USoundSubmixGraphNode* ParentNode)
 {
 	UEdGraphPin* ChildPin = ParentNode->GetChildPin();
 
@@ -292,17 +275,14 @@ USoundSubmixGraphNode* USoundSubmixGraph::CreateNode(USoundSubmix* SoundSubmix,
 
 USoundSubmixGraphNode* USoundSubmixGraph::FindExistingNode(USoundSubmix* SoundSubmix) const
 {
-	USoundSubmixGraphNode* ExistingNode = NULL;
-
 	for (int32 NodeIndex = 0; NodeIndex < Nodes.Num(); ++NodeIndex)
 	{
 		USoundSubmixGraphNode* Node = CastChecked<USoundSubmixGraphNode>(Nodes[NodeIndex]);
 		if (Node->SoundSubmix == SoundSubmix)
 		{
-			ExistingNode = Node;
-			break;
+			return Node;
 		}
 	}
 
-	return ExistingNode;
+	return nullptr;
 }
diff --git a/Engine/Source/Editor/AudioEditor/Private/SoundSubmixGraphNode.cpp b/Engine/Source/Editor/AudioEditor/Private/SoundSubmixGraphNode.cpp
--- a/Engine/Source/Editor/AudioEditor/Private/SoundSubmixGraphNode.cpp
+++ b/Engine/Source/Editor/AudioEditor/Private/SoundSubmixGraphNode.cpp
@@ -61,8 +61,8 @@ void USoundSubmixGraphNode::AllocateDefaultPins()
 {
 	check(Pins.Num() == 0);
 
-	ChildPin = CreatePin(EGPD_Output, TEXT("SoundSubmix"), *LOCTEXT("SoundSubmixChildren", "Children").ToString());
-	ParentPin = CreatePin(EGPD_Input, TEXT("SoundSubmix"), NAME_None);
+	ChildPin = CreatePin(EGPD_Input, TEXT("SoundSubmix"), *LOCTEXT("SoundSubmixChildren", "Input").ToString());
+	ParentPin = CreatePin(EGPD_Output, TEXT("SoundSubmix"), *LOCTEXT("SoundSubmixChildren", "Output").ToString());
 }
 
 void USoundSubmixGraphNode::AutowireNewNode(UEdGraphPin* FromPin)
@@ -98,13 +98,4 @@ FText USoundSubmixGraphNode::GetNodeTitle(ENodeTitleType::Type TitleType) const
 		return Super::GetNodeTitle(TitleType);
 	}
 }
-
-bool USoundSubmixGraphNode::CanUserDeleteNode() const
-{
-	USoundSubmixGraph* SoundSubmixGraph = CastChecked<USoundSubmixGraph>(GetGraph());
-
-	// Cannot remove the root node from the graph
-	return SoundSubmix != SoundSubmixGraph->GetRootSoundSubmix();
-}
-
 #undef LOCTEXT_NAMESPACE
diff --git a/Engine/Source/Editor/AudioEditor/Private/SoundSubmixGraphSchema.cpp b/Engine/Source/Editor/AudioEditor/Private/SoundSubmixGraphSchema.cpp
--- a/Engine/Source/Editor/AudioEditor/Private/SoundSubmixGraphSchema.cpp
+++ b/Engine/Source/Editor/AudioEditor/Private/SoundSubmixGraphSchema.cpp
@@ -1,23 +1,97 @@
 // Copyright 1998-2019 Epic Games, Inc. All Rights Reserved.
 
 #include "SoundSubmixGraph/SoundSubmixGraphSchema.h"
-#include "SoundSubmixGraph/SoundSubmixGraphNode.h"
-#include "SoundSubmixGraph/SoundSubmixGraph.h"
+
 #include "AssetData.h"
-#include "ScopedTransaction.h"
 #include "GraphEditorActions.h"
-#include "SoundSubmixEditorUtilities.h"
-#include "Sound/SoundSubmix.h"
+#include "EdGraph/EdGraphSchema.h"
+#include "EdGraphUtilities.h"
 #include "Framework/Commands/GenericCommands.h"
 #include "Framework/MultiBox/MultiBoxBuilder.h"
+#include "ScopedTransaction.h"
+#include "Sound/SoundSubmix.h"
+#include "SoundSubmixGraph/SoundSubmixGraphNode.h"
+#include "SoundSubmixGraph/SoundSubmixGraph.h"
+#include "SoundSubmixEditorUtilities.h"
 #include "ToolMenus.h"
 
 #define LOCTEXT_NAMESPACE "SoundSubmixSchema"
 
-UEdGraphNode* FSoundSubmixGraphSchemaAction_NewNode::PerformAction(class UEdGraph* ParentGraph, UEdGraphPin* FromPin, const FVector2D Location, bool bSelectNewNode/* = true*/)
+
+namespace
+{
+	static const FLinearColor SubmixGraphColor = FColor(175, 255, 0);
+} // namespace <>
+
+
+FConnectionDrawingPolicy* FSoundSubmixGraphConnectionDrawingPolicyFactory::CreateConnectionPolicy(
+	const UEdGraphSchema* Schema,
+	int32 InBackLayerID,
+	int32 InFrontLayerID,
+	float ZoomFactor,
+	const FSlateRect& InClippingRect,
+	FSlateWindowElementList& InDrawElements,
+	UEdGraph* InGraphObj) const
+{
+	if (Schema->IsA(USoundSubmixGraphSchema::StaticClass()))
+	{
+		return new FSoundSubmixGraphConnectionDrawingPolicy(InBackLayerID, InFrontLayerID, ZoomFactor, InClippingRect, InDrawElements, InGraphObj);
+	}
+	return nullptr;
+}
+
+FSoundSubmixGraphConnectionDrawingPolicy::FSoundSubmixGraphConnectionDrawingPolicy(int32 InBackLayerID, int32 InFrontLayerID, float ZoomFactor, const FSlateRect& InClippingRect, FSlateWindowElementList& InDrawElements, UEdGraph* InGraphObj)
+	: FConnectionDrawingPolicy(InBackLayerID, InFrontLayerID, ZoomFactor, InClippingRect, InDrawElements)
+	, GraphObj(InGraphObj)
+{
+	ActiveWireThickness = Settings->TraceAttackWireThickness;
+	InactiveWireThickness = Settings->TraceReleaseWireThickness;
+}
+
+// Give specific editor modes a chance to highlight this connection or darken non-interesting connections
+void FSoundSubmixGraphConnectionDrawingPolicy::DetermineWiringStyle(UEdGraphPin* OutputPin, UEdGraphPin* InputPin, FConnectionParams& OutParams)
+{
+	OutParams.AssociatedPin1 = InputPin;
+	OutParams.AssociatedPin2 = OutputPin;
+
+	// Get the schema and grab the default color from it
+	check(OutputPin);
+	check(GraphObj);
+	const UEdGraphSchema* Schema = GraphObj->GetSchema();
+
+	OutParams.WireColor = Schema->GetPinTypeColor(OutputPin->PinType);
+
+	if (!OutputPin)
+	{
+		return;
+	}
+
+	bool bExecuted = false;
+
+	// Run through the predecessors, and on
+	if (FExecPairingMap* PredecessorMap = PredecessorNodes.Find(OutputPin->GetOwningNode()))
+	{
+		if (FTimePair* Times = PredecessorMap->Find(InputPin->GetOwningNode()))
+		{
+			bExecuted = true;
+
+			OutParams.WireThickness = ActiveWireThickness;
+			OutParams.WireColor = SubmixGraphColor;
+			OutParams.bDrawBubbles = true;
+		}
+	}
+
+	if (!bExecuted)
+	{
+		OutParams.WireColor = SubmixGraphColor;
+		OutParams.WireThickness = InactiveWireThickness;
+	}
+}
+
+UEdGraphNode* FSoundSubmixGraphSchemaAction_NewNode::PerformAction(UEdGraph* ParentGraph, UEdGraphPin* FromPin, const FVector2D Location, bool bSelectNewNode/* = true*/)
 {
 	FSoundSubmixEditorUtilities::CreateSoundSubmix(ParentGraph, FromPin, Location, NewSoundSubmixName);
-	return NULL;
+	return nullptr;
 }
 
 USoundSubmixGraphSchema::USoundSubmixGraphSchema(const FObjectInitializer& ObjectInitializer)
@@ -30,7 +104,7 @@ bool USoundSubmixGraphSchema::ConnectionCausesLoop(const UEdGraphPin* InputPin,
 	USoundSubmixGraphNode* InputNode = CastChecked<USoundSubmixGraphNode>(InputPin->GetOwningNode());
 	USoundSubmixGraphNode* OutputNode = CastChecked<USoundSubmixGraphNode>(OutputPin->GetOwningNode());
 
-	return InputNode->SoundSubmix->RecurseCheckChild(OutputNode->SoundSubmix);
+	return OutputNode->SoundSubmix->RecurseCheckChild(InputNode->SoundSubmix);
 }
 
 void USoundSubmixGraphSchema::GetBreakLinkToSubMenuActions(UToolMenu* Menu, const FName SectionName, UEdGraphPin* InGraphPin)
@@ -41,20 +115,20 @@ void USoundSubmixGraphSchema::GetBreakLinkToSubMenuActions(UToolMenu* Menu, cons
 	TMap<FString, uint32> LinkTitleCount;
 
 	// Add all the links we could break from
-	for(TArray<class UEdGraphPin*>::TConstIterator Links(InGraphPin->LinkedTo); Links; ++Links)
+	for (TArray<UEdGraphPin*>::TConstIterator Links(InGraphPin->LinkedTo); Links; ++Links)
 	{
 		UEdGraphPin* Pin = *Links;
 		FString TitleString = Pin->GetOwningNode()->GetNodeTitle(ENodeTitleType::ListView).ToString();
-		FText Title = FText::FromString( TitleString );
+		FText Title = FText::FromString(TitleString);
 		if (Pin->PinName != TEXT(""))
 		{
 			TitleString = FString::Printf(TEXT("%s (%s)"), *TitleString, *Pin->PinName.ToString());
 
 			// Add name of connection if possible
 			FFormatNamedArguments Args;
-			Args.Add( TEXT("NodeTitle"), Title );
-			Args.Add( TEXT("PinName"), Pin->GetDisplayName() );
-			Title = FText::Format( LOCTEXT("BreakDescPin", "{NodeTitle} ({PinName})"), Args );
+			Args.Add(TEXT("NodeTitle"), Title);
+			Args.Add(TEXT("PinName"), Pin->GetDisplayName());
+			Title = FText::Format(LOCTEXT("BreakDescPin", "{NodeTitle} ({PinName})"), Args);
 		}
 
 		uint32 &Count = LinkTitleCount.FindOrAdd(TitleString);
@@ -66,11 +140,11 @@ void USoundSubmixGraphSchema::GetBreakLinkToSubMenuActions(UToolMenu* Menu, cons
 
 		if (Count == 0)
 		{
-			Description = FText::Format( LOCTEXT("BreakDesc", "Break link to {NodeTitle}"), Args );
+			Description = FText::Format(LOCTEXT("BreakDesc", "Break link to {NodeTitle}"), Args);
 		}
 		else
 		{
-			Description = FText::Format( LOCTEXT("BreakDescMulti", "Break link to {NodeTitle} ({NumberOfNodes})"), Args );
+			Description = FText::Format(LOCTEXT("BreakDescMulti", "Break link to {NodeTitle} ({NumberOfNodes})"), Args);
 		}
 		++Count;
 
@@ -140,8 +214,8 @@ const FPinConnectionResponse USoundSubmixGraphSchema::CanCreateConnection(const
 	}
 
 	// Compare the directions
-	const UEdGraphPin* InputPin = NULL;
-	const UEdGraphPin* OutputPin = NULL;
+	const UEdGraphPin* InputPin = nullptr;
+	const UEdGraphPin* OutputPin = nullptr;
 
 	if (!CategorizePinsByDirection(PinA, PinB, /*out*/ InputPin, /*out*/ OutputPin))
 	{
@@ -153,19 +227,19 @@ const FPinConnectionResponse USoundSubmixGraphSchema::CanCreateConnection(const
 		return FPinConnectionResponse(CONNECT_RESPONSE_DISALLOW, LOCTEXT("ConnectionLoop", "Connection would cause loop"));
 	}
 
-	// Break existing connections on inputs only - multiple output connections are acceptable
-	if (InputPin->LinkedTo.Num() > 0)
+	// Break existing connections on outputs only - multiple input connections are acceptable
+	if (OutputPin->LinkedTo.Num() > 0)
 	{
-		ECanCreateConnectionResponse ReplyBreakOutputs;
-		if (InputPin == PinA)
+		ECanCreateConnectionResponse ReplyBreakInputs;
+		if (OutputPin == PinA)
 		{
-			ReplyBreakOutputs = CONNECT_RESPONSE_BREAK_OTHERS_A;
+			ReplyBreakInputs = CONNECT_RESPONSE_BREAK_OTHERS_A;
 		}
 		else
 		{
-			ReplyBreakOutputs = CONNECT_RESPONSE_BREAK_OTHERS_B;
+			ReplyBreakInputs = CONNECT_RESPONSE_BREAK_OTHERS_B;
 		}
-		return FPinConnectionResponse(ReplyBreakOutputs, LOCTEXT("ConnectionReplace", "Replace existing connections"));
+		return FPinConnectionResponse(ReplyBreakInputs, LOCTEXT("ConnectionReplace", "Replace existing connections"));
 	}
 
 	return FPinConnectionResponse(CONNECT_RESPONSE_MAKE, FText::GetEmpty());
@@ -190,7 +264,7 @@ bool USoundSubmixGraphSchema::ShouldHidePinDefaultValue(UEdGraphPin* Pin) const
 
 FLinearColor USoundSubmixGraphSchema::GetPinTypeColor(const FEdGraphPinType& PinType) const
 {
-	return FLinearColor::White;
+	return SubmixGraphColor;
 }
 
 void USoundSubmixGraphSchema::BreakNodeLinks(UEdGraphNode& TargetNode) const
@@ -237,10 +311,9 @@ void USoundSubmixGraphSchema::DroppedAssetsOnGraph(const TArray<struct FAssetDat
 
 	if (UndisplayedSubmixes.Num() > 0)
 	{
-		const FScopedTransaction Transaction( LOCTEXT("SoundSubmixEditorDropSubmixes", "Sound Submix Editor: Drag and Drop Sound Submix") );
+		const FScopedTransaction Transaction(LOCTEXT("SoundSubmixEditorDropSubmixes", "Sound Submix Editor: Drag and Drop Sound Submix"));
 
 		SoundSubmixGraph->AddDroppedSoundSubmixes(UndisplayedSubmixes, GraphPosition.X, GraphPosition.Y);
 	}
 }
-
 #undef LOCTEXT_NAMESPACE
diff --git a/Engine/Source/Runtime/AudioMixer/Private/AudioMixerDevice.cpp b/Engine/Source/Runtime/AudioMixer/Private/AudioMixerDevice.cpp
--- a/Engine/Source/Runtime/AudioMixer/Private/AudioMixerDevice.cpp
+++ b/Engine/Source/Runtime/AudioMixer/Private/AudioMixerDevice.cpp
@@ -5,18 +5,19 @@
 #include "AudioMixerSubmix.h"
 #include "AudioMixerSourceVoice.h"
 #include "AudioPluginUtilities.h"
-#include "UObject/UObjectHash.h"
 #include "AudioMixerEffectsManager.h"
+#include "Sound/AudioSettings.h"
+#include "Sound/SoundSubmixSend.h"
 #include "SubmixEffects/AudioMixerSubmixEffectReverb.h"
 #include "SubmixEffects/AudioMixerSubmixEffectReverbFast.h"
 #include "SubmixEffects/AudioMixerSubmixEffectEQ.h"
 #include "SubmixEffects/AudioMixerSubmixEffectDynamicsProcessor.h"
 #include "DSP/Noise.h"
 #include "DSP/SinOsc.h"
+#include "UObject/UObjectHash.h"
 #include "UObject/UObjectIterator.h"
 #include "Runtime/HeadMountedDisplay/Public/IHeadMountedDisplayModule.h"
 #include "Misc/App.h"
-#include "Sound/AudioSettings.h"
 #include "ProfilingDebugging/CsvProfiler.h"
 #include "Async/Async.h"
 
diff --git a/Engine/Source/Runtime/AudioMixer/Public/AudioMixerBlueprintLibrary.h b/Engine/Source/Runtime/AudioMixer/Public/AudioMixerBlueprintLibrary.h
--- a/Engine/Source/Runtime/AudioMixer/Public/AudioMixerBlueprintLibrary.h
+++ b/Engine/Source/Runtime/AudioMixer/Public/AudioMixerBlueprintLibrary.h
@@ -8,6 +8,7 @@
 #include "Sound/SoundEffectSource.h"
 #include "SampleBuffer.h"
 #include "Sound/SoundCue.h"
+#include "Sound/SoundSubmixSend.h"
 #include "DSP/SpectrumAnalyzer.h"
 #include "AudioMixerBlueprintLibrary.generated.h"
 
diff --git a/Engine/Source/Runtime/AudioMixer/Public/AudioMixerSubmix.h b/Engine/Source/Runtime/AudioMixer/Public/AudioMixerSubmix.h
--- a/Engine/Source/Runtime/AudioMixer/Public/AudioMixerSubmix.h
+++ b/Engine/Source/Runtime/AudioMixer/Public/AudioMixerSubmix.h
@@ -9,6 +9,7 @@
 #include "DSP/SpectrumAnalyzer.h"
 
 class USoundEffectSubmix;
+class USoundSubmix;
 
 namespace Audio
 {
diff --git a/Engine/Source/Runtime/Engine/Classes/Sound/SoundBase.h b/Engine/Source/Runtime/Engine/Classes/Sound/SoundBase.h
--- a/Engine/Source/Runtime/Engine/Classes/Sound/SoundBase.h
+++ b/Engine/Source/Runtime/Engine/Classes/Sound/SoundBase.h
@@ -7,18 +7,19 @@
 
 #include "Audio.h"
 #include "CoreMinimal.h"
-#include "UObject/Object.h"
-#include "UObject/ObjectMacros.h"
+#include "IAudioExtensionPlugin.h"
 #include "SoundConcurrency.h"
 #include "SoundSourceBusSend.h"
-#include "SoundSubmix.h"
-#include "IAudioExtensionPlugin.h"
+#include "SoundSubmixSend.h"
+#include "UObject/Object.h"
+#include "UObject/ObjectMacros.h"
 
 #include "SoundBase.generated.h"
 
 
 class USoundEffectSourcePreset;
 class USoundSourceBus;
+class USoundSubmix;
 class USoundEffectSourcePresetChain;
 
 struct FActiveSound;
diff --git a/Engine/Source/Runtime/Engine/Classes/Sound/SoundSubmix.h b/Engine/Source/Runtime/Engine/Classes/Sound/SoundSubmix.h
--- a/Engine/Source/Runtime/Engine/Classes/Sound/SoundSubmix.h
+++ b/Engine/Source/Runtime/Engine/Classes/Sound/SoundSubmix.h
@@ -8,14 +8,18 @@
 #include "SoundEffectSubmix.h"
 #include "IAmbisonicsMixer.h"
 #include "Curves/CurveFloat.h"
+
 #include "SoundSubmix.generated.h"
 
 
+// Forward Declarations
+class FOnSubmixEnvelopeBP;
 class UEdGraph;
 class USoundEffectSubmixPreset;
 class USoundSubmix;
 class ISubmixBufferListener;
 
+
 /* Submix channel format.
  Allows submixes to have sources mix to a particular channel configuration for potential effect chain requirements.
  Master submix will always render at the device channel count. All child submixes will be down-mixed (or up-mixed) to
@@ -46,83 +50,6 @@ enum class ESubmixChannelFormat : uint8
 	Count UMETA(Hidden)
 };
 
-UENUM(BlueprintType)
-enum class EAudioRecordingExportType : uint8
-{
-	// Exports a USoundWave.
-	SoundWave,
-
-	// Exports a WAV file.
-	WavFile
-};
-
-UENUM(BlueprintType)
-enum class ESendLevelControlMethod : uint8
-{
-	// A send based on linear interpolation between a distance range and send-level range
-	Linear,
-
-	// A send based on a supplied curve
-	CustomCurve,
-
-	// A manual send level (Uses the specified constant send level value. Useful for 2D sounds.)
-	Manual,
-};
-
-
-// Class used to send audio to submixes from USoundBase
-USTRUCT(BlueprintType)
-struct ENGINE_API FSoundSubmixSendInfo
-{
-	GENERATED_USTRUCT_BODY()
-
-	/* 
-		Manual: Use Send Level only
-		Linear: Interpolate between Min and Max Send Levels based on listener distance (between Distance Min and Distance Max)
-		Custom Curve: Use the float curve to map Send Level to distance (0.0-1.0 on curve maps to Distance Min - Distance Max)
-	*/
-	UPROPERTY(EditAnywhere, Category = SubmixSend)
-	ESendLevelControlMethod SendLevelControlMethod;
-
-	// The submix to send the audio to
-	UPROPERTY(EditAnywhere, Category = SubmixSend)
-	USoundSubmix* SoundSubmix;
-
-	// The amount of audio to send
-	UPROPERTY(EditAnywhere, Category = SubmixSend)
-	float SendLevel;
-
-	// The amount to send to master when sound is located at a distance equal to value specified in the min send distance.
-	UPROPERTY(EditAnywhere, BlueprintReadWrite, Category = SubmixSend)
-	float MinSendLevel;
-
-	// The amount to send to master when sound is located at a distance equal to value specified in the max send distance.
-	UPROPERTY(EditAnywhere, BlueprintReadWrite, Category = SubmixSend)
-	float MaxSendLevel;
-
-	// The min distance to send to the master
-	UPROPERTY(EditAnywhere, BlueprintReadWrite, Category = SubmixSend)
-	float MinSendDistance;
-
-	// The max distance to send to the master
-	UPROPERTY(EditAnywhere, BlueprintReadWrite, Category = SubmixSend)
-	float MaxSendDistance;
-
-	// The custom reverb send curve to use for distance-based send level.
-	UPROPERTY(EditAnywhere, BlueprintReadWrite, Category = SubmixSend)
-	FRuntimeFloatCurve CustomSendLevelCurve;
-
-	FSoundSubmixSendInfo()
-		: SendLevelControlMethod(ESendLevelControlMethod::Manual)
-		, SoundSubmix(nullptr)
-		, SendLevel(0.0f)
-		, MinSendLevel(0.0f)
-		, MaxSendLevel(1.0f)
-		, MinSendDistance(100.0f)
-		, MaxSendDistance(1000.0f)
-		{
-		}
-};
 
 /**
 * Called when a recorded file has finished writing to disk.
@@ -135,8 +62,6 @@ DECLARE_DYNAMIC_MULTICAST_DELEGATE_OneParam(FOnSubmixRecordedFileDone, const USo
 */
 DECLARE_DYNAMIC_MULTICAST_DELEGATE_OneParam(FOnSubmixEnvelope, const TArray<float>&, Envelope);
 
-DECLARE_DYNAMIC_DELEGATE_OneParam(FOnSubmixEnvelopeBP, const TArray<float>&, Envelope);
-
 
 #if WITH_EDITOR
 
@@ -165,7 +90,7 @@ class ENGINE_API USoundSubmix : public UObject
 
 #if WITH_EDITORONLY_DATA
 	/** EdGraph based representation of the SoundSubmix */
-	class UEdGraph* SoundSubmixGraph;
+	UEdGraph* SoundSubmixGraph;
 #endif
 
 	// Experimental! Specifies the channel format for the submix. Sources will be mixed at the specified format. Useful for specific effects that need to operate on a specific format.
@@ -240,10 +165,12 @@ protected:
 	virtual FString GetDesc() override;
 	virtual void BeginDestroy() override;
 	virtual void PostLoad() override;
+
 #if WITH_EDITOR
+	virtual void PostDuplicate(EDuplicateMode::Type DuplicateMode) override;
 	virtual void PreEditChange(UProperty* PropertyAboutToChange) override;
 	virtual void PostEditChangeProperty(struct FPropertyChangedEvent& PropertyChangedEvent) override;
-#endif
+#endif // WITH_EDITOR
 	//~ End UObject Interface.
 
 	// State handling for bouncing output.
@@ -274,27 +201,7 @@ public:
 	*/
 	static void AddReferencedObjects(UObject* InThis, FReferenceCollector& Collector);
 
-	/**
-	* Refresh all EdGraph representations of SoundSubmixes
-	*
-	* @param	bIgnoreThis	Whether to ignore this SoundSubmix if it's already up to date
-	*/
-	void RefreshAllGraphs(bool bIgnoreThis);
-
-	/** Sets the sound submix graph editor implementation.* */
-	static void SetSoundSubmixAudioEditor(TSharedPtr<ISoundSubmixAudioEditor> InSoundSubmixAudioEditor);
-
-	/** Gets the sound submix graph editor implementation. */
-	static TSharedPtr<ISoundSubmixAudioEditor> GetSoundSubmixAudioEditor();
-
 private:
-
-	/** Ptr to interface to sound class editor operations. */
-	static TSharedPtr<ISoundSubmixAudioEditor> SoundSubmixAudioEditor;
-
-#endif
-
-
-
+	TArray<USoundSubmix*> BackupChildSubmixes;
+#endif // WITH_EDITOR
 };
-
diff --git a/Engine/Source/Runtime/Engine/Classes/Sound/SoundSubmixSend.h b/Engine/Source/Runtime/Engine/Classes/Sound/SoundSubmixSend.h
new file mode 100644
--- /dev/null
+++ b/Engine/Source/Runtime/Engine/Classes/Sound/SoundSubmixSend.h
@@ -0,0 +1,85 @@
+// Copyright 1998-2019 Epic Games, Inc. All Rights Reserved.
+#pragma once
+
+#include "CoreMinimal.h"
+
+#include "Curves/CurveFloat.h"
+#include "UObject/ObjectMacros.h"
+#include "UObject/Object.h"
+
+#include "SoundSubmixSend.generated.h"
+
+
+// Forward Declarations
+class USoundSubmix;
+
+
+DECLARE_DYNAMIC_DELEGATE_OneParam(FOnSubmixEnvelopeBP, const TArray<float>&, Envelope);
+
+UENUM(BlueprintType)
+enum class EAudioRecordingExportType : uint8
+{
+	// Exports a USoundWave.
+	SoundWave,
+
+	// Exports a WAV file.
+	WavFile
+};
+
+UENUM(BlueprintType)
+enum class ESendLevelControlMethod : uint8
+{
+	// A send based on linear interpolation between a distance range and send-level range
+	Linear,
+
+	// A send based on a supplied curve
+	CustomCurve,
+
+	// A manual send level (Uses the specified constant send level value. Useful for 2D sounds.)
+	Manual,
+};
+
+// Class used to send audio to submixes from USoundBase
+USTRUCT(BlueprintType)
+struct ENGINE_API FSoundSubmixSendInfo
+{
+	GENERATED_USTRUCT_BODY()
+
+	FSoundSubmixSendInfo();
+
+	/*
+		Manual: Use Send Level only
+		Linear: Interpolate between Min and Max Send Levels based on listener distance (between Distance Min and Distance Max)
+		Custom Curve: Use the float curve to map Send Level to distance (0.0-1.0 on curve maps to Distance Min - Distance Max)
+	*/
+	UPROPERTY(EditAnywhere, Category = SubmixSend)
+	ESendLevelControlMethod SendLevelControlMethod;
+
+	// The submix to send the audio to
+	UPROPERTY(EditAnywhere, Category = SubmixSend)
+	USoundSubmix* SoundSubmix;
+
+	// The amount of audio to send
+	UPROPERTY(EditAnywhere, Category = SubmixSend)
+	float SendLevel;
+
+	// The amount to send to master when sound is located at a distance equal to value specified in the min send distance.
+	UPROPERTY(EditAnywhere, BlueprintReadWrite, Category = SubmixSend)
+	float MinSendLevel;
+
+	// The amount to send to master when sound is located at a distance equal to value specified in the max send distance.
+	UPROPERTY(EditAnywhere, BlueprintReadWrite, Category = SubmixSend)
+	float MaxSendLevel;
+
+	// The min distance to send to the master
+	UPROPERTY(EditAnywhere, BlueprintReadWrite, Category = SubmixSend)
+	float MinSendDistance;
+
+	// The max distance to send to the master
+	UPROPERTY(EditAnywhere, BlueprintReadWrite, Category = SubmixSend)
+	float MaxSendDistance;
+
+	// The custom reverb send curve to use for distance-based send level.
+	UPROPERTY(EditAnywhere, BlueprintReadWrite, Category = SubmixSend)
+	FRuntimeFloatCurve CustomSendLevelCurve;
+};
diff --git a/Engine/Source/Runtime/Engine/Private/AudioDevice.cpp b/Engine/Source/Runtime/Engine/Private/AudioDevice.cpp
--- a/Engine/Source/Runtime/Engine/Private/AudioDevice.cpp
+++ b/Engine/Source/Runtime/Engine/Private/AudioDevice.cpp
@@ -4524,6 +4524,31 @@ void FAudioDevice::RetriggerVirtualLoop(FAudioVirtualLoop& VirtualLoopToRetrigge
 	AddNewActiveSoundInternal(VirtualLoopToRetrigger.GetActiveSound(), &VirtualLoopToRetrigger);
 }
 
+void FAudioDevice::AddEnvelopeFollowerDelegate(USoundSubmix* InSubmix, const FOnSubmixEnvelopeBP& OnSubmixEnvelopeBP)
+{
+	UE_LOG(LogAudio, Error, TEXT("Envelope following submixes only works with the audio mixer. Please run using -audiomixer or set INI file to use submix recording."));
+}
+
+void FAudioDevice::StartSpectrumAnalysis(USoundSubmix* InSubmix, const Audio::FSpectrumAnalyzerSettings& InSettings)
+{
+	UE_LOG(LogAudio, Error, TEXT("Spectrum analysis of submixes only works with the audio mixer. Please run using -audiomixer or set INI file to use submix recording."));
+}
+
+void FAudioDevice::StopSpectrumAnalysis(USoundSubmix* InSubmix)
+{
+	UE_LOG(LogAudio, Error, TEXT("Spectrum analysis of submixes only works with the audio mixer. Please run using -audiomixer or set INI file to use submix recording."));
+}
+
+void FAudioDevice::GetMagnitudesForFrequencies(USoundSubmix* InSubmix, const TArray<float>& InFrequencies, TArray<float>& OutMagnitudes)
+{
+	UE_LOG(LogAudio, Error, TEXT("Spectrum analysis of submixes only works with the audio mixer. Please run using -audiomixer or set INI file to use submix recording."));
+}
+
+void FAudioDevice::GetPhasesForFrequencies(USoundSubmix* InSubmix, const TArray<float>& InFrequencies, TArray<float>& OutPhases)
+{
+	UE_LOG(LogAudio, Error, TEXT("Spectrum analysis of submixes only works with the audio mixer. Please run using -audiomixer or set INI file to use submix recording."));
+}
+
 void FAudioDevice::AddVirtualLoop(const FAudioVirtualLoop& InVirtualLoop)
 {
 	FAudioVirtualLoop VirtualLoop = InVirtualLoop;
diff --git a/Engine/Source/Runtime/Engine/Private/SoundSubmix.cpp b/Engine/Source/Runtime/Engine/Private/SoundSubmix.cpp
--- a/Engine/Source/Runtime/Engine/Private/SoundSubmix.cpp
+++ b/Engine/Source/Runtime/Engine/Private/SoundSubmix.cpp
@@ -1,24 +1,21 @@
 // Copyright 1998-2019 Epic Games, Inc. All Rights Reserved.
-
-
 #include "Sound/SoundSubmix.h"
+
+#include "AudioDevice.h"
 #include "AudioDeviceManager.h"
-#include "EngineGlobals.h"
 #include "Engine/Engine.h"
+#include "EngineGlobals.h"
 #include "UObject/UObjectIterator.h"
-#include "AudioDevice.h"
+#include "Sound/SoundSubmixSend.h"
 
 #if WITH_EDITOR
 #include "Framework/Notifications/NotificationManager.h"
 #include "Widgets/Notifications/SNotificationList.h"
-#endif
-
-#if WITH_EDITOR
-TSharedPtr<ISoundSubmixAudioEditor> USoundSubmix::SoundSubmixAudioEditor = nullptr;
-#endif
+#endif // WITH_EDITOR
 
 USoundSubmix::USoundSubmix(const FObjectInitializer& ObjectInitializer)
 	: Super(ObjectInitializer)
+	, SoundSubmixGraph(nullptr)
 {
 	EnvelopeFollowerAttackTime = 10;
 	EnvelopeFollowerReleaseTime = 500;
@@ -236,9 +233,6 @@ void USoundSubmix::PostLoad()
 }
 
 #if WITH_EDITOR
-
-TArray<USoundSubmix*> BackupChildSubmixes;
-
 void USoundSubmix::PreEditChange(UProperty* PropertyAboutToChange)
 {
 	static FName NAME_ChildSubmixes(TEXT("ChildSubmixes"));
@@ -250,7 +244,7 @@ void USoundSubmix::PreEditChange(UProperty* PropertyAboutToChange)
 	}
 }
 
-void USoundSubmix::PostEditChangeProperty(struct FPropertyChangedEvent& PropertyChangedEvent)
+void USoundSubmix::PostEditChangeProperty(FPropertyChangedEvent& PropertyChangedEvent)
 {
 	// Whether or not we need to reinit the submix. Not all properties require reinitialization.
 	bool bReinitSubmix = true;
@@ -297,8 +291,6 @@ void USoundSubmix::PostEditChangeProperty(struct FPropertyChangedEvent& Property
 					BackupChildSubmixes[ChildIndex]->ParentSubmix = nullptr;
 				}
 			}
-
-			RefreshAllGraphs(false);
 		}
 		else if (PropertyChangedEvent.Property->GetFName() == NAME_ParentSubmix)
 		{
@@ -324,7 +316,6 @@ void USoundSubmix::PostEditChangeProperty(struct FPropertyChangedEvent& Property
 			}
 
 			Modify();
-			RefreshAllGraphs(false);
 		}
 		else if (PropertyChangedEvent.Property->GetFName() == NAME_OutputVolume)
 		{
@@ -349,9 +340,22 @@ void USoundSubmix::PostEditChangeProperty(struct FPropertyChangedEvent& Property
 		}
 	}
 
+	BackupChildSubmixes.Reset();
+
 	Super::PostEditChangeProperty(PropertyChangedEvent);
 }
 
+void USoundSubmix::PostDuplicate(EDuplicateMode::Type DuplicateMode)
+{
+	if (DuplicateMode == EDuplicateMode::Normal)
+	{
+		SetParentSubmix(nullptr);
+		ChildSubmixes.Reset();
+	}
+
+	Super::PostDuplicate(DuplicateMode);
+}
+
 bool USoundSubmix::RecurseCheckChild(USoundSubmix* ChildSoundSubmix)
 {
 	for (int32 Index = 0; Index < ChildSubmixes.Num(); Index++)
@@ -394,37 +398,11 @@ void USoundSubmix::AddReferencedObjects(UObject* InThis, FReferenceCollector& Co
 
 	Collector.AddReferencedObject(This->SoundSubmixGraph, This);
 
-	Super::AddReferencedObjects(InThis, Collector);
-}
-
-void USoundSubmix::RefreshAllGraphs(bool bIgnoreThis)
-{
-	if (SoundSubmixAudioEditor.IsValid())
+	for (USoundSubmix* Backup : This->BackupChildSubmixes)
 	{
-		// Update the graph representation of every SoundClass
-		for (TObjectIterator<USoundSubmix> It; It; ++It)
-		{
-			USoundSubmix* SoundSubmix = *It;
-			if (!bIgnoreThis || SoundSubmix != this)
-			{
-				if (SoundSubmix->SoundSubmixGraph)
-				{
-					SoundSubmixAudioEditor->RefreshGraphLinks(SoundSubmix->SoundSubmixGraph);
-				}
-			}
-		}
+		Collector.AddReferencedObject(Backup);
 	}
-}
 
-void USoundSubmix::SetSoundSubmixAudioEditor(TSharedPtr<ISoundSubmixAudioEditor> InSoundSubmixAudioEditor)
-{
-	check(!SoundSubmixAudioEditor.IsValid());
-	SoundSubmixAudioEditor = InSoundSubmixAudioEditor;
-}
-
-TSharedPtr<ISoundSubmixAudioEditor> USoundSubmix::GetSoundSubmixAudioEditor()
-{
-	return SoundSubmixAudioEditor;
+	Super::AddReferencedObjects(InThis, Collector);
 }
-
-#endif
+#endif // WITH_EDITOR
diff --git a/Engine/Source/Runtime/Engine/Private/SoundSubmixSend.cpp b/Engine/Source/Runtime/Engine/Private/SoundSubmixSend.cpp
new file mode 100644
--- /dev/null
+++ b/Engine/Source/Runtime/Engine/Private/SoundSubmixSend.cpp
@@ -0,0 +1,15 @@
+// Copyright 1998-2019 Epic Games, Inc. All Rights Reserved.
+#include "Sound/SoundSubmixSend.h"
+
+#include "Sound/SoundSubmix.h"
+
+FSoundSubmixSendInfo::FSoundSubmixSendInfo()
+	: SendLevelControlMethod(ESendLevelControlMethod::Manual)
+	, SoundSubmix(nullptr)
+	, SendLevel(0.0f)
+	, MinSendLevel(0.0f)
+	, MaxSendLevel(1.0f)
+	, MinSendDistance(100.0f)
+	, MaxSendDistance(1000.0f)
+	{
+	}
diff --git a/Engine/Source/Runtime/Engine/Public/ActiveSound.h b/Engine/Source/Runtime/Engine/Public/ActiveSound.h
--- a/Engine/Source/Runtime/Engine/Public/ActiveSound.h
+++ b/Engine/Source/Runtime/Engine/Public/ActiveSound.h
@@ -14,7 +14,6 @@
 #include "IAudioExtensionPlugin.h"
 #include "Sound/AudioVolume.h"
 #include "Sound/SoundConcurrency.h"
-#include "Sound/SoundSubmix.h"
 #include "Sound/SoundSourceBus.h"
 
 class FAudioDevice;
diff --git a/Engine/Source/Runtime/Engine/Public/Audio.h b/Engine/Source/Runtime/Engine/Public/Audio.h
--- a/Engine/Source/Runtime/Engine/Public/Audio.h
+++ b/Engine/Source/Runtime/Engine/Public/Audio.h
@@ -11,9 +11,9 @@
 #include "Stats/Stats.h"
 #include "HAL/ThreadSafeBool.h"
 #include "Sound/SoundClass.h"
-#include "Sound/SoundSubmix.h"
 #include "Sound/SoundAttenuation.h"
 #include "Sound/SoundEffectSource.h"
+#include "Sound/SoundSubmixSend.h"
 #include "Sound/SoundSourceBusSend.h"
 #include "IAudioExtensionPlugin.h"
 
diff --git a/Engine/Source/Runtime/Engine/Public/AudioDevice.h b/Engine/Source/Runtime/Engine/Public/AudioDevice.h
--- a/Engine/Source/Runtime/Engine/Public/AudioDevice.h
+++ b/Engine/Source/Runtime/Engine/Public/AudioDevice.h
@@ -16,6 +16,7 @@
 #include "Sound/SoundClass.h"
 #include "Sound/SoundConcurrency.h"
 #include "Sound/SoundMix.h"
+#include "Sound/SoundSubmixSend.h"
 #include "Sound/SoundSourceBus.h"
 #include "AudioVirtualLoop.h"
 #include "AudioMixer.h"
@@ -1202,30 +1203,11 @@ public:
 	}
 
 	/** Adds an envelope follower delegate to the submix for this audio device. */
-	virtual void AddEnvelopeFollowerDelegate(USoundSubmix* InSubmix, const FOnSubmixEnvelopeBP& OnSubmixEnvelopeBP)
-	{
-		UE_LOG(LogAudio, Error, TEXT("Envelope following submixes only works with the audio mixer. Please run using -audiomixer or set INI file to use submix recording."));
-	}
-
-	virtual void StartSpectrumAnalysis(USoundSubmix* InSubmix, const Audio::FSpectrumAnalyzerSettings& InSettings)
-	{
-		UE_LOG(LogAudio, Error, TEXT("Spectrum analysis of submixes only works with the audio mixer. Please run using -audiomixer or set INI file to use submix recording."));
-	}
-
-	virtual void StopSpectrumAnalysis(USoundSubmix* InSubmix)
-	{
-		UE_LOG(LogAudio, Error, TEXT("Spectrum analysis of submixes only works with the audio mixer. Please run using -audiomixer or set INI file to use submix recording."));
-	}
-
-	virtual void GetMagnitudesForFrequencies(USoundSubmix* InSubmix, const TArray<float>& InFrequencies, TArray<float>& OutMagnitudes)
-	{
-		UE_LOG(LogAudio, Error, TEXT("Spectrum analysis of submixes only works with the audio mixer. Please run using -audiomixer or set INI file to use submix recording."));
-	}
-
-	virtual void GetPhasesForFrequencies(USoundSubmix* InSubmix, const TArray<float>& InFrequencies, TArray<float>& OutPhases)
-	{
-		UE_LOG(LogAudio, Error, TEXT("Spectrum analysis of submixes only works with the audio mixer. Please run using -audiomixer or set INI file to use submix recording."));
-	}
+	virtual void AddEnvelopeFollowerDelegate(USoundSubmix* InSubmix, const FOnSubmixEnvelopeBP& OnSubmixEnvelopeBP);
+	virtual void StartSpectrumAnalysis(USoundSubmix* InSubmix, const Audio::FSpectrumAnalyzerSettings& InSettings);
+	virtual void StopSpectrumAnalysis(USoundSubmix* InSubmix);
+	virtual void GetMagnitudesForFrequencies(USoundSubmix* InSubmix, const TArray<float>& InFrequencies, TArray<float>& OutMagnitudes);
+	virtual void GetPhasesForFrequencies(USoundSubmix* InSubmix, const TArray<float>& InFrequencies, TArray<float>& OutPhases);
 
 protected:
 	friend class FSoundSource;
diff --git a/Engine/Source/Runtime/Engine/Public/IAudioExtensionPlugin.h b/Engine/Source/Runtime/Engine/Public/IAudioExtensionPlugin.h
--- a/Engine/Source/Runtime/Engine/Public/IAudioExtensionPlugin.h
+++ b/Engine/Source/Runtime/Engine/Public/IAudioExtensionPlugin.h
@@ -9,7 +9,6 @@
 #include "Math/Interval.h"
 #include "Modules/ModuleInterface.h"
 #include "Modules/ModuleManager.h"
-#include "Sound/SoundSubmix.h"
 #include "UObject/ObjectMacros.h"
 
 #if !UE_BUILD_SHIPPING
@@ -19,8 +18,10 @@
 
 #include "IAudioExtensionPlugin.generated.h"
 
-
+// Forward Declarations
 class FAudioDevice;
+class USoundSubmix;
+
 
 /**
 * Enumeration of audio plugin types
@@ -773,7 +774,7 @@ public:
 	/** Called when a source is done playing and is released. */
 	virtual void OnReleaseSource(const uint32 SourceId) = 0;
 
-	virtual class FSoundEffectSubmix* GetEffectSubmix(class USoundSubmix* Submix) = 0;
+	virtual class FSoundEffectSubmix* GetEffectSubmix(USoundSubmix* Submix) = 0;
 
 	/** Processes audio with the given input and output data structs.*/
 	virtual void ProcessSourceAudio(const FAudioPluginSourceInputData& InputData, FAudioPluginSourceOutputData& OutputData)
