Missing type action update from last submix improvements CL

diff --git a/Engine/Source/Editor/AudioEditor/Private/AssetTypeActions/AssetTypeActions_SoundSubmix.cpp b/Engine/Source/Editor/AudioEditor/Private/AssetTypeActions/AssetTypeActions_SoundSubmix.cpp
--- a/Engine/Source/Editor/AudioEditor/Private/AssetTypeActions/AssetTypeActions_SoundSubmix.cpp
+++ b/Engine/Source/Editor/AudioEditor/Private/AssetTypeActions/AssetTypeActions_SoundSubmix.cpp
@@ -1,8 +1,10 @@
 // Copyright 1998-2019 Epic Games, Inc. All Rights Reserved.
 
 #include "AssetTypeActions_SoundSubmix.h"
-#include "Sound/SoundSubmix.h"
 #include "AudioEditorModule.h"
+#include "Modules/ModuleManager.h"
+#include "Sound/SoundSubmix.h"
+#include "SoundSubmixEditor.h"
 
 #define LOCTEXT_NAMESPACE "AssetTypeActions"
 
@@ -11,7 +13,7 @@ UClass* FAssetTypeActions_SoundSubmix::GetSupportedClass() const
 	return USoundSubmix::StaticClass();
 }
 
-void FAssetTypeActions_SoundSubmix::OpenAssetEditor( const TArray<UObject*>& InObjects, TSharedPtr<IToolkitHost> EditWithinLevelEditor )
+void FAssetTypeActions_SoundSubmix::OpenAssetEditor(const TArray<UObject*>& InObjects, TSharedPtr<IToolkitHost> EditWithinLevelEditor)
 {
 	EToolkitMode::Type Mode = EditWithinLevelEditor.IsValid() ? EToolkitMode::WorldCentric : EToolkitMode::Standalone;
 
@@ -20,17 +22,38 @@ void FAssetTypeActions_SoundSubmix::OpenAssetEditor( const TArray<UObject*>& InO
 		USoundSubmix* SoundSubmix = Cast<USoundSubmix>(*ObjIt);
 		if (SoundSubmix != nullptr)
 		{
-			IAudioEditorModule* AudioEditorModule = &FModuleManager::LoadModuleChecked<IAudioEditorModule>( "AudioEditor" );
+			IAudioEditorModule* AudioEditorModule = &FModuleManager::LoadModuleChecked<IAudioEditorModule>("AudioEditor");
 			AudioEditorModule->CreateSoundSubmixEditor(Mode, EditWithinLevelEditor, SoundSubmix);
 		}
 	}
 }
 
+void FAssetTypeActions_SoundSubmix::AssetsActivated(const TArray<UObject*>& InObjects, EAssetTypeActivationMethod::Type ActivationType)
+{
+	FAssetTypeActions_Base::AssetsActivated(InObjects, ActivationType);
+
+	TSet<USoundSubmix*> SubmixesToSelect;
+	IAssetEditorInstance* Editor = nullptr;
+	for (UObject* Obj : InObjects)
+	{
+		if (USoundSubmix* SubmixToSelect = Cast<USoundSubmix>(Obj))
+		{
+			if (!Editor)
+			{
+				Editor = FAssetEditorManager::Get().FindEditorForAsset(Obj, false);
+			}
+			SubmixesToSelect.Add(SubmixToSelect);
+		}
+	}
+
+	static_cast<FSoundSubmixEditor*>(Editor)->SelectSubmixes(SubmixesToSelect);
+}
+
 const TArray<FText>& FAssetTypeActions_SoundSubmix::GetSubMenus() const
 {
 	static const TArray<FText> SubMenus
 	{
-		FText(LOCTEXT("AssetSoundMixSubMenu", "Mix"))
+		LOCTEXT("AssetSoundMixSubMenu", "Mix")
 	};
 
 	return SubMenus;
diff --git a/Engine/Source/Editor/AudioEditor/Private/AssetTypeActions/AssetTypeActions_SoundSubmix.h b/Engine/Source/Editor/AudioEditor/Private/AssetTypeActions/AssetTypeActions_SoundSubmix.h
--- a/Engine/Source/Editor/AudioEditor/Private/AssetTypeActions/AssetTypeActions_SoundSubmix.h
+++ b/Engine/Source/Editor/AudioEditor/Private/AssetTypeActions/AssetTypeActions_SoundSubmix.h
@@ -12,6 +12,8 @@ public:
 	virtual FColor GetTypeColor() const override { return FColor(175, 255, 0); }
 	virtual const TArray<FText>& GetSubMenus() const override;
 	virtual UClass* GetSupportedClass() const override;
+
+	virtual void AssetsActivated(const TArray<UObject*>& InObjects, EAssetTypeActivationMethod::Type ActivationType) override;
 	virtual void OpenAssetEditor( const TArray<UObject*>& InObjects, TSharedPtr<class IToolkitHost> EditWithinLevelEditor = TSharedPtr<IToolkitHost>() ) override;
 	virtual uint32 GetCategories() override { return EAssetTypeCategories::Sounds; }
 };
\ No newline at end of file
