Fix potential crash
#rnx

diff --git a/Engine/Source/Editor/AudioEditor/Private/AssetTypeActions/AssetTypeActions_SoundSubmix.cpp b/Engine/Source/Editor/AudioEditor/Private/AssetTypeActions/AssetTypeActions_SoundSubmix.cpp
--- a/Engine/Source/Editor/AudioEditor/Private/AssetTypeActions/AssetTypeActions_SoundSubmix.cpp
+++ b/Engine/Source/Editor/AudioEditor/Private/AssetTypeActions/AssetTypeActions_SoundSubmix.cpp
@@ -46,7 +46,10 @@ void FAssetTypeActions_SoundSubmix::AssetsActivated(const TArray<UObject*>& InOb
 		}
 	}
 
-	static_cast<FSoundSubmixEditor*>(Editor)->SelectSubmixes(SubmixesToSelect);
+	if (Editor)
+	{
+		static_cast<FSoundSubmixEditor*>(Editor)->SelectSubmixes(SubmixesToSelect);
+	}
 }
 
 const TArray<FText>& FAssetTypeActions_SoundSubmix::GetSubMenus() const
