Fix deprecation warnings
#rnx

diff --git a/Engine/Source/Editor/AudioEditor/Private/AssetTypeActions/AssetTypeActions_SoundSubmix.cpp b/Engine/Source/Editor/AudioEditor/Private/AssetTypeActions/AssetTypeActions_SoundSubmix.cpp
--- a/Engine/Source/Editor/AudioEditor/Private/AssetTypeActions/AssetTypeActions_SoundSubmix.cpp
+++ b/Engine/Source/Editor/AudioEditor/Private/AssetTypeActions/AssetTypeActions_SoundSubmix.cpp
@@ -28,10 +28,8 @@ void FAssetTypeActions_SoundSubmix::OpenAssetEditor(const TArray<UObject*>& InOb
 	}
 }
 
-void FAssetTypeActions_SoundSubmix::AssetsActivated(const TArray<UObject*>& InObjects, EAssetTypeActivationMethod::Type ActivationType)
+bool FAssetTypeActions_SoundSubmix::AssetsActivatedOverride(const TArray<UObject*>& InObjects, EAssetTypeActivationMethod::Type ActivationType)
 {
-	FAssetTypeActions_Base::AssetsActivated(InObjects, ActivationType);
-
 	TSet<USoundSubmix*> SubmixesToSelect;
 	IAssetEditorInstance* Editor = nullptr;
 	for (UObject* Obj : InObjects)
@@ -40,7 +38,7 @@ void FAssetTypeActions_SoundSubmix::AssetsActivated(const TArray<UObject*>& InOb
 		{
 			if (!Editor)
 			{
-				Editor = FAssetEditorManager::Get().FindEditorForAsset(Obj, false);
+				Editor = GEditor->GetEditorSubsystem<UAssetEditorSubsystem>()->FindEditorForAsset(Obj, false);
 			}
 			SubmixesToSelect.Add(SubmixToSelect);
 		}
@@ -49,7 +47,10 @@ void FAssetTypeActions_SoundSubmix::AssetsActivated(const TArray<UObject*>& InOb
 	if (Editor)
 	{
 		static_cast<FSoundSubmixEditor*>(Editor)->SelectSubmixes(SubmixesToSelect);
+		return true;
 	}
+
+	return false;
 }
 
 const TArray<FText>& FAssetTypeActions_SoundSubmix::GetSubMenus() const
diff --git a/Engine/Source/Editor/AudioEditor/Private/AssetTypeActions/AssetTypeActions_SoundSubmix.h b/Engine/Source/Editor/AudioEditor/Private/AssetTypeActions/AssetTypeActions_SoundSubmix.h
--- a/Engine/Source/Editor/AudioEditor/Private/AssetTypeActions/AssetTypeActions_SoundSubmix.h
+++ b/Engine/Source/Editor/AudioEditor/Private/AssetTypeActions/AssetTypeActions_SoundSubmix.h
@@ -13,7 +13,7 @@ public:
 	virtual const TArray<FText>& GetSubMenus() const override;
 	virtual UClass* GetSupportedClass() const override;
 
-	virtual void AssetsActivated(const TArray<UObject*>& InObjects, EAssetTypeActivationMethod::Type ActivationType) override;
+	virtual bool AssetsActivatedOverride(const TArray<UObject*>& InObjects, EAssetTypeActivationMethod::Type ActivationType) override;
 	virtual void OpenAssetEditor( const TArray<UObject*>& InObjects, TSharedPtr<class IToolkitHost> EditWithinLevelEditor = TSharedPtr<IToolkitHost>() ) override;
 	virtual uint32 GetCategories() override { return EAssetTypeCategories::Sounds; }
 };
\ No newline at end of file
