Fix static analysis warning
#rnx

diff --git a/Engine/Source/Editor/AudioEditor/Private/SoundSubmixGraphSchema.cpp b/Engine/Source/Editor/AudioEditor/Private/SoundSubmixGraphSchema.cpp
--- a/Engine/Source/Editor/AudioEditor/Private/SoundSubmixGraphSchema.cpp
+++ b/Engine/Source/Editor/AudioEditor/Private/SoundSubmixGraphSchema.cpp
@@ -61,11 +61,6 @@ void FSoundSubmixGraphConnectionDrawingPolicy::DetermineWiringStyle(UEdGraphPin*
 
 	OutParams.WireColor = Schema->GetPinTypeColor(OutputPin->PinType);
 
-	if (!OutputPin)
-	{
-		return;
-	}
-
 	bool bExecuted = false;
 
 	// Run through the predecessors, and on
