- Safeguards around only having a single editor open for a given submix (avoids having multiple similar but not same graphs open for the same asset, which gets super confusing visually and hard to track in code)
- Enforce not being able to delete an editor's root submix from the graph (to avoid having to rebase the editable objects in the graph)

diff --git a/Engine/Source/Editor/AudioEditor/Classes/SoundSubmixGraph/SoundSubmixGraph.h b/Engine/Source/Editor/AudioEditor/Classes/SoundSubmixGraph/SoundSubmixGraph.h
--- a/Engine/Source/Editor/AudioEditor/Classes/SoundSubmixGraph/SoundSubmixGraph.h
+++ b/Engine/Source/Editor/AudioEditor/Classes/SoundSubmixGraph/SoundSubmixGraph.h
@@ -38,7 +38,7 @@ public:
 	 * @param	NodePosX		X coordinate submix(es) were dropped at
 	 * @param	NodePosY		Y coordinate submix(es) were dropped at
 	 */
-	void AddDroppedSoundSubmixes(const TArray<USoundSubmix*>& SoundSubmixes, int32 NodePosX, int32 NodePosY);
+	void AddDroppedSoundSubmixes(const TSet<USoundSubmix*>& SoundSubmixes, int32 NodePosX, int32 NodePosY);
 
 	/**
 	 * Display a new SoundSubmix that has just been created using the editor
diff --git a/Engine/Source/Editor/AudioEditor/Classes/SoundSubmixGraph/SoundSubmixGraphNode.h b/Engine/Source/Editor/AudioEditor/Classes/SoundSubmixGraph/SoundSubmixGraphNode.h
--- a/Engine/Source/Editor/AudioEditor/Classes/SoundSubmixGraph/SoundSubmixGraphNode.h
+++ b/Engine/Source/Editor/AudioEditor/Classes/SoundSubmixGraph/SoundSubmixGraphNode.h
@@ -32,6 +32,7 @@ class USoundSubmixGraphNode : public UEdGraphNode
 	virtual void AllocateDefaultPins() override;
 	virtual void AutowireNewNode(UEdGraphPin* FromPin) override;
 	virtual bool CanCreateUnderSpecifiedSchema(const UEdGraphSchema* Schema) const override;
+	virtual bool CanUserDeleteNode() const override;
 	virtual FText GetNodeTitle(ENodeTitleType::Type TitleType) const override;
 	//~ End UEdGraphNode Interface.
 
diff --git a/Engine/Source/Editor/AudioEditor/Private/AssetTypeActions/AssetTypeActions_SoundSubmix.cpp b/Engine/Source/Editor/AudioEditor/Private/AssetTypeActions/AssetTypeActions_SoundSubmix.cpp
--- a/Engine/Source/Editor/AudioEditor/Private/AssetTypeActions/AssetTypeActions_SoundSubmix.cpp
+++ b/Engine/Source/Editor/AudioEditor/Private/AssetTypeActions/AssetTypeActions_SoundSubmix.cpp
@@ -17,10 +17,9 @@ void FAssetTypeActions_SoundSubmix::OpenAssetEditor(const TArray<UObject*>& InOb
 {
 	EToolkitMode::Type Mode = EditWithinLevelEditor.IsValid() ? EToolkitMode::WorldCentric : EToolkitMode::Standalone;
 
-	for (auto ObjIt = InObjects.CreateConstIterator(); ObjIt; ++ObjIt)
+	for (UObject* Obj : InObjects)
 	{
-		USoundSubmix* SoundSubmix = Cast<USoundSubmix>(*ObjIt);
-		if (SoundSubmix != nullptr)
+		if (USoundSubmix* SoundSubmix = Cast<USoundSubmix>(Obj))
 		{
 			IAudioEditorModule* AudioEditorModule = &FModuleManager::LoadModuleChecked<IAudioEditorModule>("AudioEditor");
 			AudioEditorModule->CreateSoundSubmixEditor(Mode, EditWithinLevelEditor, SoundSubmix);
diff --git a/Engine/Source/Editor/AudioEditor/Private/SoundSubmixEditor.cpp b/Engine/Source/Editor/AudioEditor/Private/SoundSubmixEditor.cpp
--- a/Engine/Source/Editor/AudioEditor/Private/SoundSubmixEditor.cpp
+++ b/Engine/Source/Editor/AudioEditor/Private/SoundSubmixEditor.cpp
@@ -5,6 +5,7 @@
 #include "AssetRegistryModule.h"
 #include "AssetToolsModule.h"
 #include "AudioEditorModule.h"
+#include "EdGraph/EdGraph.h"
 #include "Editor.h"
 #include "GraphEditor.h"
 #include "GraphEditAction.h"
@@ -26,6 +27,7 @@
 #include "UObject/Package.h"
 #include "Widgets/Docking/SDockTab.h"
 
+
 #define LOCTEXT_NAMESPACE "SoundSubmixEditor"
 DEFINE_LOG_CATEGORY_STATIC(LogSoundSubmixEditor, Log, All);
 
@@ -36,7 +38,7 @@ const FName FSoundSubmixEditor::PropertiesTabId(TEXT("SoundSubmixEditor_Properti
 class SSoundSubmixGraphEditor : public SGraphEditor
 {
 private:
-	FSoundSubmixEditor* SubmixEditor;
+	TWeakPtr<FSoundSubmixEditor> SubmixEditor;
 
 public:
 	SSoundSubmixGraphEditor()
@@ -45,15 +47,22 @@ public:
 	{
 	}
 
-	void Construct(const FArguments& InArgs, FSoundSubmixEditor& InEditor)
+	void Construct(const FArguments& InArgs, TSharedPtr<FSoundSubmixEditor> InEditor)
 	{
-		SubmixEditor = &InEditor;
+		SubmixEditor = InEditor;
 		SGraphEditor::Construct(InArgs);
 	}
 
 	void OnGraphChanged(const FEdGraphEditAction& InAction) override
 	{
-		SubmixEditor->AddMissingEditableSubmixes(InAction.Graph);
+		if (SubmixEditor.IsValid())
+		{
+			TSharedPtr<FSoundSubmixEditor> PinnedEditor = SubmixEditor.Pin();
+			if (InAction.Graph && InAction.Graph == PinnedEditor->GetGraph())
+			{
+				PinnedEditor->AddMissingEditableSubmixes();
+			}
+		}
 		SGraphEditor::OnGraphChanged(InAction);
 	}
 };
@@ -122,14 +131,11 @@ void FSoundSubmixEditor::Init(const EToolkitMode::Type Mode, const TSharedPtr<IT
 		FGenericCommands::Get().Redo,
 		FExecuteAction::CreateSP(this, &FSoundSubmixEditor::RedoGraphAction));
 
-	if (!SoundSubmix->SoundSubmixGraph)
-	{
-		USoundSubmixGraph* SoundSubmixGraph = CastChecked<USoundSubmixGraph>(FBlueprintEditorUtils::CreateNewGraph(SoundSubmix, NAME_None, USoundSubmixGraph::StaticClass(), USoundSubmixGraphSchema::StaticClass()));
-		SoundSubmixGraph->SetRootSoundSubmix(SoundSubmix);
+	USoundSubmixGraph* SoundSubmixGraph = CastChecked<USoundSubmixGraph>(FBlueprintEditorUtils::CreateNewGraph(SoundSubmix, NAME_None, USoundSubmixGraph::StaticClass(), USoundSubmixGraphSchema::StaticClass()));
+	SoundSubmixGraph->SetRootSoundSubmix(SoundSubmix);
 
-		SoundSubmix->SoundSubmixGraph = SoundSubmixGraph;
-		SoundSubmixGraph->RebuildGraph();
-	}
+	SoundSubmix->SoundSubmixGraph = SoundSubmixGraph;
+	SoundSubmixGraph->RebuildGraph();
 
 	CreateInternalWidgets(SoundSubmix);
 
@@ -288,7 +294,7 @@ TSharedRef<SGraphEditor> FSoundSubmixEditor::CreateGraphEditorWidget(USoundSubmi
 	InEvents.OnSelectionChanged = SGraphEditor::FOnSelectionChanged::CreateSP(this, &FSoundSubmixEditor::OnSelectedNodesChanged);
 	InEvents.OnCreateActionMenu = SGraphEditor::FOnCreateActionMenu::CreateSP(this, &FSoundSubmixEditor::OnCreateGraphActionMenu);
 
-	return SNew(SSoundSubmixGraphEditor, *this)
+	return SNew(SSoundSubmixGraphEditor, SharedThis(this))
 		.AdditionalCommands(GraphEditorCommands)
 		.IsEditable(true)
 		.Appearance(AppearanceInfo)
@@ -369,7 +375,7 @@ void FSoundSubmixEditor::RemoveSelectedNodes()
 {
 	const FScopedTransaction Transaction(LOCTEXT("SoundSubmixEditorRemoveSelectedNode", "Sound Submix Editor: Remove Selected SoundSubmixes from editor"));
 
-	int32 NumObjectsRemoved = GetEditingObjects().Num();
+	int32 NumObjectsRemoved = 0;
 
 	const TSet<UObject*> SelectedNodes = GraphEditor->GetSelectedNodes();
 	for (UObject* SelectedNode : SelectedNodes)
@@ -377,16 +383,14 @@ void FSoundSubmixEditor::RemoveSelectedNodes()
 		USoundSubmixGraphNode* Node = Cast<USoundSubmixGraphNode>(SelectedNode);
 		if (Node && Node->SoundSubmix && Node->CanUserDeleteNode())
 		{
-			NumObjectsRemoved--;
+			NumObjectsRemoved++;
 			RemoveEditingObject(Node->SoundSubmix);
 		}
 	}
 
-	check(NumObjectsRemoved >= 0);
 	if (NumObjectsRemoved > 0)
 	{
-		USoundSubmix* SoundSubmix = CastChecked<USoundSubmix>(GetEditingObjects()[0]);
-		USoundSubmixGraph* Graph = CastChecked<USoundSubmixGraph>(SoundSubmix->SoundSubmixGraph);
+		USoundSubmixGraph* Graph = CastChecked<USoundSubmixGraph>(GraphEditor->GetCurrentGraph());
 		Graph->RecursivelyRemoveNodes(SelectedNodes);
 		GraphEditor->ClearSelectionSet();
 	}
@@ -446,38 +450,41 @@ void FSoundSubmixEditor::CreateSoundSubmix(UEdGraphPin* FromPin, const FVector2D
 	}
 }
 
+UEdGraph* FSoundSubmixEditor::GetGraph()
+{
+	return GraphEditor->GetCurrentGraph();
+}
+
 FText FSoundSubmixEditor::GetToolkitName() const
 {
 	UObject* EditObject = GetEditingObjects()[0];
 	return GetLabelForObject(EditObject);
 }
 
-void FSoundSubmixEditor::AddMissingEditableSubmixes(UEdGraph* InGraph)
+void FSoundSubmixEditor::AddMissingEditableSubmixes()
 {
-	UEdGraph* Graph = GraphEditor->GetCurrentGraph();
-	check(Graph);
-
-	if (InGraph != Graph)
+	if (UEdGraph* Graph = GraphEditor->GetCurrentGraph())
 	{
-		return;
-	}
-
-	bool bRefreshGraphs = false;
-	if (Graph->Nodes.Num() > GetEditingObjects().Num())
-	{
-		for (UEdGraphNode* Node : Graph->Nodes)
+		bool bChanged = false;
+		if (Graph->Nodes.Num() > GetEditingObjects().Num())
 		{
-			USoundSubmixGraphNode* GraphNode = CastChecked<USoundSubmixGraphNode>(Node);
-			USoundSubmix* UntrackedSubmix = GraphNode->SoundSubmix;
-			if (UntrackedSubmix && !GetEditingObjects().Contains(UntrackedSubmix))
+			for (UEdGraphNode* Node : Graph->Nodes)
 			{
-				AddEditingObject(UntrackedSubmix);
-				bRefreshGraphs = true;
+				USoundSubmixGraphNode* GraphNode = CastChecked<USoundSubmixGraphNode>(Node);
+				USoundSubmix* UntrackedSubmix = GraphNode->SoundSubmix;
+				if (UntrackedSubmix && !GetEditingObjects().Contains(UntrackedSubmix))
+				{
+					bChanged = true;
+					AddEditingObject(UntrackedSubmix);
+				}
 			}
 		}
-	}
 
-	GraphEditor->NotifyGraphChanged();
+		if (bChanged)
+		{
+			GraphEditor->NotifyGraphChanged();
+		}
+	}
 }
 
 void FSoundSubmixEditor::PostUndo(bool bSuccess)
diff --git a/Engine/Source/Editor/AudioEditor/Private/SoundSubmixEditor.h b/Engine/Source/Editor/AudioEditor/Private/SoundSubmixEditor.h
--- a/Engine/Source/Editor/AudioEditor/Private/SoundSubmixEditor.h
+++ b/Engine/Source/Editor/AudioEditor/Private/SoundSubmixEditor.h
@@ -14,16 +14,11 @@
 // Forward Declarations
 class UEdGraphPin;
 class USoundSubmix;
+class UEdGraph;
 
 class FSoundSubmixEditor : public ISoundSubmixEditor, public FGCObject, public FEditorUndoClient
 {
 public:
-	SLATE_BEGIN_ARGS(FSoundSubmixEditor)
-	{
-	}
-
-	SLATE_END_ARGS()
-
 	virtual void RegisterTabSpawners(const TSharedRef<FTabManager>& TabManager) override;
 	virtual void UnregisterTabSpawners(const TSharedRef<FTabManager>& TabManager) override;
 
@@ -59,16 +54,18 @@ public:
 	virtual void PostUndo(bool bSuccess) override;
 	virtual void PostRedo(bool bSuccess) override { PostUndo(bSuccess); }
 
-	void AddMissingEditableSubmixes(UEdGraph* InGraph);
+	void AddMissingEditableSubmixes();
 
 	/** Select node associated with the provided submix */
 	void SelectSubmixes(TSet<USoundSubmix*>& InSubmixes);
 
+	/** Returns current graph handled by editor */
+	UEdGraph* GetGraph();
+
 private:
 	TSharedRef<SDockTab> SpawnTab_GraphCanvas(const FSpawnTabArgs& Args);
 	TSharedRef<SDockTab> SpawnTab_Properties(const FSpawnTabArgs& Args);
 
-private:
 	/** Creates all internal widgets for the tabs to point at */
 	void CreateInternalWidgets(USoundSubmix* InSoundSubmix);
 
diff --git a/Engine/Source/Editor/AudioEditor/Private/SoundSubmixGraph.cpp b/Engine/Source/Editor/AudioEditor/Private/SoundSubmixGraph.cpp
--- a/Engine/Source/Editor/AudioEditor/Private/SoundSubmixGraph.cpp
+++ b/Engine/Source/Editor/AudioEditor/Private/SoundSubmixGraph.cpp
@@ -15,6 +15,7 @@ USoundSubmixGraph::USoundSubmixGraph(const FObjectInitializer& ObjectInitializer
 
 void USoundSubmixGraph::SetRootSoundSubmix(USoundSubmix* InSoundSubmix)
 {
+	check(!RootSoundSubmix);
 	RootSoundSubmix = InSoundSubmix;
 }
 
@@ -43,13 +44,13 @@ void USoundSubmixGraph::RebuildGraph()
 	Package->SetDirtyFlag(bIsDirty);
 }
 
-void USoundSubmixGraph::AddDroppedSoundSubmixes(const TArray<USoundSubmix*>& SoundSubmixes, int32 NodePosX, int32 NodePosY)
+void USoundSubmixGraph::AddDroppedSoundSubmixes(const TSet<USoundSubmix*>& SoundSubmixes, int32 NodePosX, int32 NodePosY)
 {
 	Modify();
 
-	for (int32 ClassIndex = 0; ClassIndex < SoundSubmixes.Num(); ClassIndex++)
+	for (USoundSubmix* SoundSubmix : SoundSubmixes)
 	{
-		NodePosY += ConstructNodes(SoundSubmixes[ClassIndex], NodePosX, NodePosY);
+		NodePosY += ConstructNodes(SoundSubmix, NodePosX, NodePosY);
 	}
 
 	NotifyGraphChanged();
@@ -83,10 +84,8 @@ void USoundSubmixGraph::LinkSoundSubmixes()
 			Node->SoundSubmix->Modify();
 
 			// remove parents of existing children
-			for (int32 ChildIndex = 0; ChildIndex < Node->SoundSubmix->ChildSubmixes.Num(); ChildIndex++)
+			for (USoundSubmix* ChildSubmix : Node->SoundSubmix->ChildSubmixes)
 			{
-				USoundSubmix* ChildSubmix = Node->SoundSubmix->ChildSubmixes[ChildIndex];
-
 				if (ChildSubmix)
 				{
 					ChildSubmix->Modify();
@@ -96,13 +95,19 @@ void USoundSubmixGraph::LinkSoundSubmixes()
 
 			Node->SoundSubmix->ChildSubmixes.Empty();
 
-			UEdGraphPin* ChildPin = Node->GetChildPin();
-
-			for (int32 ChildIndex = 0; ChildIndex < ChildPin->LinkedTo.Num(); ChildIndex++)
+			if (UEdGraphPin* ChildPin = Node->GetChildPin())
 			{
-				USoundSubmixGraphNode* ChildNode = CastChecked<USoundSubmixGraphNode>(ChildPin->LinkedTo[ChildIndex]->GetOwningNode());
-				Node->SoundSubmix->ChildSubmixes.Add(ChildNode->SoundSubmix);
-				ChildNode->SoundSubmix->SetParentSubmix(Node->SoundSubmix);
+				for (UEdGraphPin* GraphPin : ChildPin->LinkedTo)
+				{
+					if (!GraphPin)
+					{
+						continue;
+					}
+
+					USoundSubmixGraphNode* ChildNode = CastChecked<USoundSubmixGraphNode>(GraphPin->GetOwningNode());
+					Node->SoundSubmix->ChildSubmixes.Add(ChildNode->SoundSubmix);
+					ChildNode->SoundSubmix->SetParentSubmix(Node->SoundSubmix);
+				}
 			}
 
 			Node->SoundSubmix->PostEditChange();
diff --git a/Engine/Source/Editor/AudioEditor/Private/SoundSubmixGraphNode.cpp b/Engine/Source/Editor/AudioEditor/Private/SoundSubmixGraphNode.cpp
--- a/Engine/Source/Editor/AudioEditor/Private/SoundSubmixGraphNode.cpp
+++ b/Engine/Source/Editor/AudioEditor/Private/SoundSubmixGraphNode.cpp
@@ -4,6 +4,8 @@
 #include "Sound/SoundSubmix.h"
 #include "SoundSubmixGraph/SoundSubmixGraphSchema.h"
 #include "SoundSubmixGraph/SoundSubmixGraph.h"
+#include "Toolkits/AssetEditorManager.h"
+#include "SoundSubmixEditor.h"
 
 #define LOCTEXT_NAMESPACE "SoundSubmixGraphNode"
 
@@ -87,6 +89,33 @@ bool USoundSubmixGraphNode::CanCreateUnderSpecifiedSchema(const UEdGraphSchema*
 	return Schema->IsA(USoundSubmixGraphSchema::StaticClass());
 }
 
+bool USoundSubmixGraphNode::CanUserDeleteNode() const
+{
+	TArray<IAssetEditorInstance*> SubmixEditors = FAssetEditorManager::Get().FindEditorsForAsset(SoundSubmix);
+	for (IAssetEditorInstance* Editor : SubmixEditors)
+	{
+		if (!Editor)
+		{
+			continue;
+		}
+
+		FSoundSubmixEditor* SubmixEditor = static_cast<FSoundSubmixEditor*>(Editor);
+		if (UEdGraph* Graph = SubmixEditor->GetGraph())
+		{
+			if (SoundSubmix->SoundSubmixGraph == Graph)
+			{
+				USoundSubmix* RootSubmix = CastChecked<USoundSubmixGraph>(Graph)->GetRootSoundSubmix();
+				if (RootSubmix == SoundSubmix)
+				{
+					return false;
+				}
+			}
+		}
+	}
+
+	return UEdGraphNode::CanUserDeleteNode();
+}
+
 FText USoundSubmixGraphNode::GetNodeTitle(ENodeTitleType::Type TitleType) const
 {
 	if (SoundSubmix)
diff --git a/Engine/Source/Editor/AudioEditor/Private/SoundSubmixGraphSchema.cpp b/Engine/Source/Editor/AudioEditor/Private/SoundSubmixGraphSchema.cpp
--- a/Engine/Source/Editor/AudioEditor/Private/SoundSubmixGraphSchema.cpp
+++ b/Engine/Source/Editor/AudioEditor/Private/SoundSubmixGraphSchema.cpp
@@ -12,9 +12,12 @@
 #include "Sound/SoundSubmix.h"
 #include "SoundSubmixGraph/SoundSubmixGraphNode.h"
 #include "SoundSubmixGraph/SoundSubmixGraph.h"
+#include "SoundSubmixEditor.h"
 #include "SoundSubmixEditorUtilities.h"
+#include "Toolkits/AssetEditorManager.h"
 #include "ToolMenus.h"
 
+
 #define LOCTEXT_NAMESPACE "SoundSubmixSchema"
 
 
@@ -242,11 +245,41 @@ const FPinConnectionResponse USoundSubmixGraphSchema::CanCreateConnection(const
 
 bool USoundSubmixGraphSchema::TryCreateConnection(UEdGraphPin* PinA, UEdGraphPin* PinB) const
 {
+	check(PinA);
+	check(PinB);
+
 	bool bModified = UEdGraphSchema::TryCreateConnection(PinA, PinB);
 
 	if (bModified)
 	{
-		CastChecked<USoundSubmixGraph>(PinA->GetOwningNode()->GetGraph())->LinkSoundSubmixes();
+		USoundSubmixGraph* Graph = CastChecked<USoundSubmixGraph>(PinA->GetOwningNode()->GetGraph());
+		Graph->LinkSoundSubmixes();
+
+		USoundSubmix* SubmixA = CastChecked<USoundSubmixGraphNode>(PinA->GetOwningNode())->SoundSubmix;
+		USoundSubmix* SubmixB = CastChecked<USoundSubmixGraphNode>(PinB->GetOwningNode())->SoundSubmix;
+
+		bool bReopenEditors = false;
+
+		// If re-basing root, re-open editor.  This will force the root to be the primary edited node
+		if (Graph->GetRootSoundSubmix() == SubmixA && SubmixA->ParentSubmix != nullptr)
+		{
+			bReopenEditors = true;
+		}
+		else if (Graph->GetRootSoundSubmix() == SubmixB && SubmixB->ParentSubmix != nullptr)
+		{
+			bReopenEditors = true;
+		}
+
+		if (bReopenEditors)
+		{
+			TArray<IAssetEditorInstance*> Editors = FAssetEditorManager::Get().FindEditorsForAsset(SubmixA);
+			for (IAssetEditorInstance* Editor : Editors)
+			{
+				Editor->CloseWindow();
+			}
+
+			FAssetEditorManager::Get().OpenEditorForAsset(SubmixA);
+		}
 	}
 
 	return bModified;
@@ -290,17 +323,35 @@ void USoundSubmixGraphSchema::BreakSinglePinLink(UEdGraphPin* SourcePin, UEdGrap
 	CastChecked<USoundSubmixGraph>(SourcePin->GetOwningNode()->GetGraph())->LinkSoundSubmixes();
 }
 
-void USoundSubmixGraphSchema::DroppedAssetsOnGraph(const TArray<struct FAssetData>& Assets, const FVector2D& GraphPosition, UEdGraph* Graph) const
+void USoundSubmixGraphSchema::DroppedAssetsOnGraph(const TArray<FAssetData>& Assets, const FVector2D& GraphPosition, UEdGraph* Graph) const
 {
-	USoundSubmixGraph* SoundSubmixGraph = CastChecked<USoundSubmixGraph>(Graph);
+	check(Graph);
 
-	TArray<USoundSubmix*> UndisplayedSubmixes;
-	for (int32 AssetIdx = 0; AssetIdx < Assets.Num(); ++AssetIdx)
+	USoundSubmixGraph* SoundSubmixGraph = CastChecked<USoundSubmixGraph>(Graph);
+	TSet<IAssetEditorInstance*> Editors;
+	TSet<USoundSubmix*> UndisplayedSubmixes;
+	for (const FAssetData& Asset : Assets)
 	{
-		USoundSubmix* SoundSubmix = Cast<USoundSubmix>(Assets[AssetIdx].GetAsset());
-		if (SoundSubmix && !SoundSubmixGraph->IsSubmixDisplayed(SoundSubmix))
+		// Walk to the root submix
+		if (USoundSubmix* SoundSubmix = Cast<USoundSubmix>(Asset.GetAsset()))
 		{
-			UndisplayedSubmixes.Add(SoundSubmix);
+			while (SoundSubmix->ParentSubmix != nullptr)
+			{
+				SoundSubmix = SoundSubmix->ParentSubmix;
+			}
+
+			if (!SoundSubmixGraph->IsSubmixDisplayed(SoundSubmix))
+			{
+				TArray<IAssetEditorInstance*> SubmixEditors = FAssetEditorManager::Get().FindEditorsForAsset(SoundSubmix);
+				for (IAssetEditorInstance* Editor : SubmixEditors)
+				{
+					if (Editor)
+					{
+						Editors.Add(Editor);
+					}
+				}
+				UndisplayedSubmixes.Add(SoundSubmix);
+			}
 		}
 	}
 
@@ -308,7 +359,30 @@ void USoundSubmixGraphSchema::DroppedAssetsOnGraph(const TArray<struct FAssetDat
 	{
 		const FScopedTransaction Transaction(LOCTEXT("SoundSubmixEditorDropSubmixes", "Sound Submix Editor: Drag and Drop Sound Submix"));
 
-		SoundSubmixGraph->AddDroppedSoundSubmixes(UndisplayedSubmixes, GraphPosition.X, GraphPosition.Y);
+		for (IAssetEditorInstance* Editor : Editors)
+		{
+			check(Editor);
+			FSoundSubmixEditor* SubmixEditor = static_cast<FSoundSubmixEditor*>(Editor);
+
+			// Close editors with dropped (and undisplayed) submix branches as they are now displayed locally in this graph
+			// (to avoid modification of multiple graph editors representing the same branch of submixes)
+			if (SubmixEditor->GetGraph() != Graph)
+			{
+				Editor->CloseWindow();
+			}
+		}
+
+		// If editor is this graph's editor, update editable objects and select dropped submixes.
+		if (USoundSubmix* RootSubmix = SoundSubmixGraph->GetRootSoundSubmix())
+		{
+			if (IAssetEditorInstance* Editor = FAssetEditorManager::Get().FindEditorForAsset(RootSubmix, false /* bFocusIfOpen */))
+			{
+				FSoundSubmixEditor* SubmixEditor = static_cast<FSoundSubmixEditor*>(Editor);
+				SoundSubmixGraph->AddDroppedSoundSubmixes(UndisplayedSubmixes, GraphPosition.X, GraphPosition.Y);
+				SubmixEditor->AddMissingEditableSubmixes();
+				SubmixEditor->SelectSubmixes(UndisplayedSubmixes);
+			}
+		}
 	}
 }
 #undef LOCTEXT_NAMESPACE
diff --git a/Engine/Source/Runtime/Engine/Classes/Sound/SoundSubmix.h b/Engine/Source/Runtime/Engine/Classes/Sound/SoundSubmix.h
--- a/Engine/Source/Runtime/Engine/Classes/Sound/SoundSubmix.h
+++ b/Engine/Source/Runtime/Engine/Classes/Sound/SoundSubmix.h
@@ -82,7 +82,7 @@ class ENGINE_API USoundSubmix : public UObject
 	GENERATED_UCLASS_BODY()
 
 	// Child submixes to this sound mix
-	UPROPERTY(EditAnywhere, BlueprintReadOnly, Category = SoundSubmix)
+	UPROPERTY(VisibleAnywhere, BlueprintReadOnly, Category = SoundSubmix)
 	TArray<USoundSubmix*> ChildSubmixes;
 
 	UPROPERTY()
