Copy fixes to 4.24: SoundSubmix Editor Improvements
- Make Submix graph point in an intuitive direction and change input/output directions.
- Fix bugs around adding/removing submixes and having multiple editors open.
- Move implementation and definitions so SoundSubmix updates builds faster (remove SoundSubmix.h references from bulky headers that are used all over the engine)
- Safeguards around only having a single editor open for a given submix (avoids having multiple similar but not same graphs open for the same asset, which gets super confusing visually and hard to track in code)
- Enforce not being able to delete an editor's root submix from the graph (to avoid having to rebase the editable objects in the graph)

diff --git a/Engine/Source/Editor/AudioEditor/Private/SoundSubmixGraphNode.cpp b/Engine/Source/Editor/AudioEditor/Private/SoundSubmixGraphNode.cpp
--- a/Engine/Source/Editor/AudioEditor/Private/SoundSubmixGraphNode.cpp
+++ b/Engine/Source/Editor/AudioEditor/Private/SoundSubmixGraphNode.cpp
@@ -91,7 +91,9 @@ bool USoundSubmixGraphNode::CanCreateUnderSpecifiedSchema(const UEdGraphSchema*
 
 bool USoundSubmixGraphNode::CanUserDeleteNode() const
 {
-	TArray<IAssetEditorInstance*> SubmixEditors = FAssetEditorManager::Get().FindEditorsForAsset(SoundSubmix);
+	check(GEditor);
+	UAssetEditorSubsystem* EditorSubsystem = GEditor->GetEditorSubsystem<UAssetEditorSubsystem>();
+	TArray<IAssetEditorInstance*> SubmixEditors = EditorSubsystem->FindEditorsForAsset(SoundSubmix);
 	for (IAssetEditorInstance* Editor : SubmixEditors)
 	{
 		if (!Editor)
diff --git a/Engine/Source/Editor/AudioEditor/Private/SoundSubmixGraphSchema.cpp b/Engine/Source/Editor/AudioEditor/Private/SoundSubmixGraphSchema.cpp
--- a/Engine/Source/Editor/AudioEditor/Private/SoundSubmixGraphSchema.cpp
+++ b/Engine/Source/Editor/AudioEditor/Private/SoundSubmixGraphSchema.cpp
@@ -64,6 +64,11 @@ void FSoundSubmixGraphConnectionDrawingPolicy::DetermineWiringStyle(UEdGraphPin*
 
 	OutParams.WireColor = Schema->GetPinTypeColor(OutputPin->PinType);
 
+	if (!OutputPin)
+	{
+		return;
+	}
+
 	bool bExecuted = false;
 
 	// Run through the predecessors, and on
@@ -272,13 +277,15 @@ bool USoundSubmixGraphSchema::TryCreateConnection(UEdGraphPin* PinA, UEdGraphPin
 
 		if (bReopenEditors)
 		{
-			TArray<IAssetEditorInstance*> Editors = FAssetEditorManager::Get().FindEditorsForAsset(SubmixA);
-			for (IAssetEditorInstance* Editor : Editors)
+			check(GEditor);
+			UAssetEditorSubsystem* EditorSubsystem = GEditor->GetEditorSubsystem<UAssetEditorSubsystem>();
+			TArray<IAssetEditorInstance*> SubmixEditors = EditorSubsystem->FindEditorsForAsset(SubmixA);
+			for (IAssetEditorInstance* Editor : SubmixEditors)
 			{
 				Editor->CloseWindow();
 			}
 
-			FAssetEditorManager::Get().OpenEditorForAsset(SubmixA);
+			EditorSubsystem->OpenEditorForAsset(SubmixA);
 		}
 	}
 
@@ -325,6 +332,7 @@ void USoundSubmixGraphSchema::BreakSinglePinLink(UEdGraphPin* SourcePin, UEdGrap
 
 void USoundSubmixGraphSchema::DroppedAssetsOnGraph(const TArray<FAssetData>& Assets, const FVector2D& GraphPosition, UEdGraph* Graph) const
 {
+	check(GEditor);
 	check(Graph);
 
 	USoundSubmixGraph* SoundSubmixGraph = CastChecked<USoundSubmixGraph>(Graph);
@@ -342,7 +350,8 @@ void USoundSubmixGraphSchema::DroppedAssetsOnGraph(const TArray<FAssetData>& Ass
 
 			if (!SoundSubmixGraph->IsSubmixDisplayed(SoundSubmix))
 			{
-				TArray<IAssetEditorInstance*> SubmixEditors = FAssetEditorManager::Get().FindEditorsForAsset(SoundSubmix);
+				UAssetEditorSubsystem* EditorSubsystem = GEditor->GetEditorSubsystem<UAssetEditorSubsystem>();
+				TArray<IAssetEditorInstance*> SubmixEditors = EditorSubsystem->FindEditorsForAsset(SoundSubmix);
 				for (IAssetEditorInstance* Editor : SubmixEditors)
 				{
 					if (Editor)
@@ -375,6 +384,8 @@ void USoundSubmixGraphSchema::DroppedAssetsOnGraph(const TArray<FAssetData>& Ass
 		// If editor is this graph's editor, update editable objects and select dropped submixes.
 		if (USoundSubmix* RootSubmix = SoundSubmixGraph->GetRootSoundSubmix())
 		{
+			UAssetEditorSubsystem* EditorSubsystem = GEditor->GetEditorSubsystem<UAssetEditorSubsystem>();
+			if (IAssetEditorInstance* Editor = EditorSubsystem->FindEditorForAsset(RootSubmix, false /* bFocusIfOpen */))
 			if (IAssetEditorInstance* Editor = FAssetEditorManager::Get().FindEditorForAsset(RootSubmix, false /* bFocusIfOpen */))
 			{
 				FSoundSubmixEditor* SubmixEditor = static_cast<FSoundSubmixEditor*>(Editor);
diff --git a/Engine/Source/Runtime/Engine/Classes/Sound/SoundSubmix.h b/Engine/Source/Runtime/Engine/Classes/Sound/SoundSubmix.h
--- a/Engine/Source/Runtime/Engine/Classes/Sound/SoundSubmix.h
+++ b/Engine/Source/Runtime/Engine/Classes/Sound/SoundSubmix.h
@@ -14,6 +14,7 @@
 
 
 // Forward Declarations
+class FOnSubmixEnvelopeBP;
 class UEdGraph;
 class USoundEffectSubmixPreset;
 class USoundSubmix;
