Bad merge & variable hiding warning fix

diff --git a/Engine/Source/Editor/AudioEditor/Private/SoundSubmixGraphSchema.cpp b/Engine/Source/Editor/AudioEditor/Private/SoundSubmixGraphSchema.cpp
--- a/Engine/Source/Editor/AudioEditor/Private/SoundSubmixGraphSchema.cpp
+++ b/Engine/Source/Editor/AudioEditor/Private/SoundSubmixGraphSchema.cpp
@@ -385,10 +385,9 @@ void USoundSubmixGraphSchema::DroppedAssetsOnGraph(const TArray<FAssetData>& Ass
 		if (USoundSubmix* RootSubmix = SoundSubmixGraph->GetRootSoundSubmix())
 		{
 			UAssetEditorSubsystem* EditorSubsystem = GEditor->GetEditorSubsystem<UAssetEditorSubsystem>();
-			if (IAssetEditorInstance* Editor = EditorSubsystem->FindEditorForAsset(RootSubmix, false /* bFocusIfOpen */))
-			if (IAssetEditorInstance* Editor = FAssetEditorManager::Get().FindEditorForAsset(RootSubmix, false /* bFocusIfOpen */))
+			if (IAssetEditorInstance* EditorInstance = EditorSubsystem->FindEditorForAsset(RootSubmix, false /* bFocusIfOpen */))
 			{
-				FSoundSubmixEditor* SubmixEditor = static_cast<FSoundSubmixEditor*>(Editor);
+				FSoundSubmixEditor* SubmixEditor = static_cast<FSoundSubmixEditor*>(EditorInstance);
 				SoundSubmixGraph->AddDroppedSoundSubmixes(UndisplayedSubmixes, GraphPosition.X, GraphPosition.Y);
 				SubmixEditor->AddMissingEditableSubmixes();
 				SubmixEditor->SelectSubmixes(UndisplayedSubmixes);
diff --git a/Engine/Source/Runtime/Engine/Classes/Sound/SoundSubmix.h b/Engine/Source/Runtime/Engine/Classes/Sound/SoundSubmix.h
--- a/Engine/Source/Runtime/Engine/Classes/Sound/SoundSubmix.h
+++ b/Engine/Source/Runtime/Engine/Classes/Sound/SoundSubmix.h
@@ -85,9 +85,11 @@ class ENGINE_API USoundSubmix : public UObject
 	UPROPERTY(VisibleAnywhere, BlueprintReadOnly, Category = SoundSubmix)
 	TArray<USoundSubmix*> ChildSubmixes;
 
+public:
 	UPROPERTY()
 	USoundSubmix* ParentSubmix;
 
+private:
 #if WITH_EDITORONLY_DATA
 	/** EdGraph based representation of the SoundSubmix */
 	UEdGraph* SoundSubmixGraph;
diff --git a/Engine/Source/Runtime/Engine/Private/SoundSubmix.cpp b/Engine/Source/Runtime/Engine/Private/SoundSubmix.cpp
--- a/Engine/Source/Runtime/Engine/Private/SoundSubmix.cpp
+++ b/Engine/Source/Runtime/Engine/Private/SoundSubmix.cpp
@@ -15,6 +15,7 @@
 
 USoundSubmix::USoundSubmix(const FObjectInitializer& ObjectInitializer)
 	: Super(ObjectInitializer)
+	, ParentSubmix(nullptr)
 #if WITH_EDITORONLY_DATA
 	, SoundSubmixGraph(nullptr)
 #endif // WITH_EDITORONLY_DATA
