Minor clean-up regarding static analysis warning

diff --git a/Engine/Source/Editor/AudioEditor/Private/SoundSubmixGraphSchema.cpp b/Engine/Source/Editor/AudioEditor/Private/SoundSubmixGraphSchema.cpp
--- a/Engine/Source/Editor/AudioEditor/Private/SoundSubmixGraphSchema.cpp
+++ b/Engine/Source/Editor/AudioEditor/Private/SoundSubmixGraphSchema.cpp
@@ -54,12 +54,13 @@ FSoundSubmixGraphConnectionDrawingPolicy::FSoundSubmixGraphConnectionDrawingPoli
 // Give specific editor modes a chance to highlight this connection or darken non-interesting connections
 void FSoundSubmixGraphConnectionDrawingPolicy::DetermineWiringStyle(UEdGraphPin* OutputPin, UEdGraphPin* InputPin, FConnectionParams& OutParams)
 {
+	check(GraphObj);
+	check(OutputPin);
+
 	OutParams.AssociatedPin1 = InputPin;
 	OutParams.AssociatedPin2 = OutputPin;
 
 	// Get the schema and grab the default color from it
-	check(OutputPin);
-	check(GraphObj);
 	const UEdGraphSchema* Schema = GraphObj->GetSchema();
 
 	OutParams.WireColor = Schema->GetPinTypeColor(OutputPin->PinType);
