Fix for first import missing playback issue: Move and shuffle invalidation/free resource calls to end of SoundFactory and always call free regardless of bUseExistingSettings flag

diff --git a/Engine/Source/Editor/AudioEditor/Private/Factories/SoundFactory.cpp b/Engine/Source/Editor/AudioEditor/Private/Factories/SoundFactory.cpp
--- a/Engine/Source/Editor/AudioEditor/Private/Factories/SoundFactory.cpp
+++ b/Engine/Source/Editor/AudioEditor/Private/Factories/SoundFactory.cpp
@@ -349,18 +349,6 @@ UObject* USoundFactory::CreateObject
 			Sound->ConcurrencySet = TemplateSoundWave->ConcurrencySet;
 		}
 
-		if (bUseExistingSettings && ExistingSound)
-		{
-			// Clear resources so that if it's already been played, it will reload the wave data
-			Sound->FreeResources();
-		}
-
-		// Store the current file path and timestamp for re-import purposes
-		Sound->AssetImportData->Update(CurrentFilename);
-
-		// Compressed data is now out of date.
-		Sound->InvalidateCompressedData();
-		 
 		// If we're a multi-channel file, we're going to spoof the behavior of the SoundSurroundFactory
 		int32 ChannelCount = (int32)*WaveInfo.pChannels;
 		check(ChannelCount >0);
@@ -498,6 +486,12 @@ UObject* USoundFactory::CreateObject
 		Sound->NumChannels = ChannelCount;
 		Sound->TotalSamples = *WaveInfo.pSamplesPerSec * Sound->Duration;
 
+		// Store the current file path and timestamp for re-import purposes
+		Sound->AssetImportData->Update(CurrentFilename);
+
+		// Compressed data is now out of date.
+		Sound->InvalidateCompressedData(true /* bFreeResources */);
+
 		// If stream caching is enabled, we need to make sure this asset is ready for playback.
 		if (FPlatformCompressionUtilities::IsCurrentPlatformUsingStreamCaching() && Sound->IsStreaming())
 		{
