- Don't prompt for template update if reimporting sound
- Log to display if stopping a sound due to re-import

diff --git a/Engine/Source/Editor/AudioEditor/Private/Factories/SoundFactory.cpp b/Engine/Source/Editor/AudioEditor/Private/Factories/SoundFactory.cpp
--- a/Engine/Source/Editor/AudioEditor/Private/Factories/SoundFactory.cpp
+++ b/Engine/Source/Editor/AudioEditor/Private/Factories/SoundFactory.cpp
@@ -244,18 +244,21 @@ UObject* USoundFactory::CreateObject
 				}
 			}
 
-			UE_LOG(LogAudioEditor, Log, TEXT("Stopping Sound Resources of Existing Sound"));
 			if (ComponentsToRestart.Num() > 0)
 			{
+				UE_LOG(LogAudioEditor, Display, TEXT("Stopping the following AudioComponents referencing sound being imported"));
 				for (UAudioComponent* AudioComponent : ComponentsToRestart)
 				{
-					UE_LOG(LogAudioEditor, Log, TEXT("Component '%s' Stopped"), *AudioComponent->GetName());
+					UE_LOG(LogAudioEditor, Display, TEXT("Component '%s' Stopped"), *AudioComponent->GetName());
 					AudioComponent->Stop();
 				}
 			}
 		}
 
-		UpdateTemplate();
+		if (!ExistingSound)
+		{
+			UpdateTemplate();
+		}
 
 		bool bUseExistingSettings = SuppressImportDialogOptions & ESuppressImportDialog::Overwrite;
 		if (ExistingSound && !bUseExistingSettings && !GIsAutomationTesting)
