Fix copy pasta dup'ed loc string

diff --git a/Engine/Source/Editor/AudioEditor/Private/SoundSubmixGraphNode.cpp b/Engine/Source/Editor/AudioEditor/Private/SoundSubmixGraphNode.cpp
--- a/Engine/Source/Editor/AudioEditor/Private/SoundSubmixGraphNode.cpp
+++ b/Engine/Source/Editor/AudioEditor/Private/SoundSubmixGraphNode.cpp
@@ -63,8 +63,8 @@ void USoundSubmixGraphNode::AllocateDefaultPins()
 {
 	check(Pins.Num() == 0);
 
-	ChildPin = CreatePin(EGPD_Input, TEXT("SoundSubmix"), *LOCTEXT("SoundSubmixChildren", "Input").ToString());
-	ParentPin = CreatePin(EGPD_Output, TEXT("SoundSubmix"), *LOCTEXT("SoundSubmixChildren", "Output").ToString());
+	ChildPin = CreatePin(EGPD_Input, TEXT("SoundSubmix"), *LOCTEXT("SoundSubmixGraphNode_Input", "Input").ToString());
+	ParentPin = CreatePin(EGPD_Output, TEXT("SoundSubmix"), *LOCTEXT("SoundSubmixGraphNode_Output", "Output").ToString());
 }
 
 void USoundSubmixGraphNode::AutowireNewNode(UEdGraphPin* FromPin)
