Fix for surround sound import issue where one channel will be silent

diff --git a/Engine/Source/Editor/AudioEditor/Private/Factories/SoundSurroundFactory.cpp b/Engine/Source/Editor/AudioEditor/Private/Factories/SoundSurroundFactory.cpp
--- a/Engine/Source/Editor/AudioEditor/Private/Factories/SoundSurroundFactory.cpp
+++ b/Engine/Source/Editor/AudioEditor/Private/Factories/SoundSurroundFactory.cpp
@@ -123,7 +123,7 @@ UObject* USoundSurroundFactory::FactoryCreateBinary
 		Sound->AssetImportData->Update(CurrentFilename);
 
 		// Compressed data is now out of date.
-		Sound->InvalidateCompressedData();
+		Sound->InvalidateCompressedData(false, false);
 
 		// Delete the old version of the wave from the bulk data
 		uint8* RawWaveData[SPEAKER_Count] = { nullptr };
