Graph Editor node menus converted to ToolMenus
#rnx

diff --git a/Engine/Plugins/Experimental/ControlRig/Source/ControlRigDeveloper/Private/Graph/ControlRigGraphNode.cpp b/Engine/Plugins/Experimental/ControlRig/Source/ControlRigDeveloper/Private/Graph/ControlRigGraphNode.cpp
--- a/Engine/Plugins/Experimental/ControlRig/Source/ControlRigDeveloper/Private/Graph/ControlRigGraphNode.cpp
+++ b/Engine/Plugins/Experimental/ControlRig/Source/ControlRigDeveloper/Private/Graph/ControlRigGraphNode.cpp
@@ -841,10 +841,10 @@ void UControlRigGraphNode::PinConnectionListChanged(UEdGraphPin* Pin)
 
 }
 
-void UControlRigGraphNode::GetContextMenuActions(const FGraphNodeContextMenuBuilder& Context) const
+void UControlRigGraphNode::GetNodeContextMenuActions(UToolMenu* Menu, UGraphNodeContextMenuContext* Context) const
 {
 #if WITH_EDITOR
-	IControlRigEditorModule::Get().GetContextMenuActions(this, Context);
+	IControlRigEditorModule::Get().GetNodeContextMenuActions(this, Menu, Context);
 #endif
 }
 
diff --git a/Engine/Plugins/Experimental/ControlRig/Source/ControlRigDeveloper/Public/Graph/ControlRigGraphNode.h b/Engine/Plugins/Experimental/ControlRig/Source/ControlRigDeveloper/Public/Graph/ControlRigGraphNode.h
--- a/Engine/Plugins/Experimental/ControlRig/Source/ControlRigDeveloper/Public/Graph/ControlRigGraphNode.h
+++ b/Engine/Plugins/Experimental/ControlRig/Source/ControlRigDeveloper/Public/Graph/ControlRigGraphNode.h
@@ -168,7 +168,7 @@ public:
 	virtual void AllocateDefaultPins() override;
 	virtual void ReconstructNode() override;
 	virtual void PinConnectionListChanged(UEdGraphPin* Pin) override;
-	virtual void GetContextMenuActions(const FGraphNodeContextMenuBuilder& Context) const override;
+	virtual void GetNodeContextMenuActions(class UToolMenu* Menu, class UGraphNodeContextMenuContext* Context) const override;
 	virtual void DestroyNode() override;
 	virtual void PinDefaultValueChanged(UEdGraphPin* Pin) override;
 	virtual TSharedPtr<INameValidatorInterface> MakeNameValidator() const override;
diff --git a/Engine/Plugins/Experimental/ControlRig/Source/ControlRigEditor/Private/ControlRigEditorModule.cpp b/Engine/Plugins/Experimental/ControlRig/Source/ControlRigEditor/Private/ControlRigEditorModule.cpp
--- a/Engine/Plugins/Experimental/ControlRig/Source/ControlRigEditor/Private/ControlRigEditorModule.cpp
+++ b/Engine/Plugins/Experimental/ControlRig/Source/ControlRigEditor/Private/ControlRigEditorModule.cpp
@@ -860,50 +860,42 @@ FConnectionDrawingPolicy* FControlRigEditorModule::CreateConnectionDrawingPolicy
 	return new FControlRigConnectionDrawingPolicy(InBackLayerID, InFrontLayerID, InZoomFactor, InClippingRect, InDrawElements, InGraphObj);
 }
 
-void FControlRigEditorModule::GetContextMenuActions(const UControlRigGraphNode* Node, const FGraphNodeContextMenuBuilder& Context )
+void FControlRigEditorModule::GetNodeContextMenuActions(const UControlRigGraphNode* Node, UToolMenu* Menu, UGraphNodeContextMenuContext* Context) const
 {
-	if(Context.MenuBuilder != nullptr)
 	{
-		if(Context.Pin != nullptr)
+		if (Context->Pin != nullptr)
 		{
 			// Add array operations for array pins
-			if(Context.Pin->PinType.IsArray())
+			if (Context->Pin->PinType.IsArray())
 			{
-				// End the section as this function is called with a section 'open'
-				Context.MenuBuilder->EndSection();
-
-				Context.MenuBuilder->BeginSection(TEXT("ArrayOperations"), LOCTEXT("ArrayOperations", "Array Operations"));
+				FToolMenuSection& Section = Menu->AddSection(TEXT("ArrayOperations"), LOCTEXT("ArrayOperations", "Array Operations"));
 
 				// Array operations
-				Context.MenuBuilder->AddMenuEntry(
+				Section.AddMenuEntry(
+					"ClearArray",
 					LOCTEXT("ClearArray", "Clear"),
 					LOCTEXT("ClearArray_Tooltip", "Clear this array of all of its entries"),
 					FSlateIcon(),
-					FUIAction(FExecuteAction::CreateUObject(const_cast<UControlRigGraphNode*>(Node), &UControlRigGraphNode::HandleClearArray, Context.Pin->PinName.ToString())));
-
-				Context.MenuBuilder->EndSection();
+					FUIAction(FExecuteAction::CreateUObject(const_cast<UControlRigGraphNode*>(Node), &UControlRigGraphNode::HandleClearArray, Context->Pin->PinName.ToString())));
 			}
-			else if(Context.Pin->ParentPin != nullptr && Context.Pin->ParentPin->PinType.IsArray())
+			else if (Context->Pin->ParentPin != nullptr && Context->Pin->ParentPin->PinType.IsArray())
 			{
-				// End the section as this function is called with a section 'open'
-				Context.MenuBuilder->EndSection();
-
-				Context.MenuBuilder->BeginSection(TEXT("ArrayElementOperations"), LOCTEXT("ArrayElementOperations", "Array Element Operations"));
+				FToolMenuSection& Section = Menu->AddSection(TEXT("ArrayElementOperations"), LOCTEXT("ArrayElementOperations", "Array Element Operations"));
 
 				// Array element operations
-				Context.MenuBuilder->AddMenuEntry(
+				Section.AddMenuEntry(
+					"RemoveArrayElement",
 					LOCTEXT("RemoveArrayElement", "Remove"),
 					LOCTEXT("RemoveArrayElement_Tooltip", "Remove this array element"),
 					FSlateIcon(),
-					FUIAction(FExecuteAction::CreateUObject(const_cast<UControlRigGraphNode*>(Node), &UControlRigGraphNode::HandleRemoveArrayElement, Context.Pin->PinName.ToString())));
+					FUIAction(FExecuteAction::CreateUObject(const_cast<UControlRigGraphNode*>(Node), &UControlRigGraphNode::HandleRemoveArrayElement, Context->Pin->PinName.ToString())));
 
-				Context.MenuBuilder->AddMenuEntry(
+				Section.AddMenuEntry(
+					"InsertArrayElement",
 					LOCTEXT("InsertArrayElement", "Insert"),
 					LOCTEXT("InsertArrayElement_Tooltip", "Insert an array element after this one"),
 					FSlateIcon(),
-					FUIAction(FExecuteAction::CreateUObject(const_cast<UControlRigGraphNode*>(Node), &UControlRigGraphNode::HandleInsertArrayElement, Context.Pin->PinName.ToString())));
-
-				Context.MenuBuilder->EndSection();
+					FUIAction(FExecuteAction::CreateUObject(const_cast<UControlRigGraphNode*>(Node), &UControlRigGraphNode::HandleInsertArrayElement, Context->Pin->PinName.ToString())));
 			}
 		}
 	}
diff --git a/Engine/Plugins/Experimental/ControlRig/Source/ControlRigEditor/Private/ControlRigEditorModule.h b/Engine/Plugins/Experimental/ControlRig/Source/ControlRigEditor/Private/ControlRigEditorModule.h
--- a/Engine/Plugins/Experimental/ControlRig/Source/ControlRigEditor/Private/ControlRigEditorModule.h
+++ b/Engine/Plugins/Experimental/ControlRig/Source/ControlRigEditor/Private/ControlRigEditorModule.h
@@ -47,7 +47,7 @@ public:
 	virtual void GetTypeActions(const UControlRigBlueprint* CRB, FBlueprintActionDatabaseRegistrar& ActionRegistrar) override;
 	virtual void GetInstanceActions(const UControlRigBlueprint* CRB, FBlueprintActionDatabaseRegistrar& ActionRegistrar) override;
 	virtual FConnectionDrawingPolicy* CreateConnectionDrawingPolicy(int32 InBackLayerID, int32 InFrontLayerID, float InZoomFactor, const FSlateRect& InClippingRect, class FSlateWindowElementList& InDrawElements, class UEdGraph* InGraphObj) override;
-	virtual void GetContextMenuActions(const UControlRigGraphNode* Node, const FGraphNodeContextMenuBuilder& Context ) override;
+	virtual void GetNodeContextMenuActions(const UControlRigGraphNode* Node, class UToolMenu* Menu, class UGraphNodeContextMenuContext* Context) const override;
 	virtual void GetContextMenuActions(const UControlRigGraphSchema* Schema, class UToolMenu* Menu, class UGraphNodeContextMenuContext* Context) const override;
 
 	static TSubclassOf<URigUnitEditor_Base> GetEditorObjectByRigUnit(const FName& RigUnitClassName);
diff --git a/Engine/Plugins/Experimental/ControlRig/Source/ControlRigEditor/Public/IControlRigEditorModule.h b/Engine/Plugins/Experimental/ControlRig/Source/ControlRigEditor/Public/IControlRigEditorModule.h
--- a/Engine/Plugins/Experimental/ControlRig/Source/ControlRigEditor/Public/IControlRigEditorModule.h
+++ b/Engine/Plugins/Experimental/ControlRig/Source/ControlRigEditor/Public/IControlRigEditorModule.h
@@ -46,6 +46,6 @@ public:
 	virtual void GetTypeActions(const UControlRigBlueprint* CRB, FBlueprintActionDatabaseRegistrar& ActionRegistrar) = 0;
 	virtual void GetInstanceActions(const UControlRigBlueprint* CRB, FBlueprintActionDatabaseRegistrar& ActionRegistrar) = 0;
 	virtual FConnectionDrawingPolicy* CreateConnectionDrawingPolicy(int32 InBackLayerID, int32 InFrontLayerID, float InZoomFactor, const FSlateRect& InClippingRect, class FSlateWindowElementList& InDrawElements, class UEdGraph* InGraphObj) = 0;
-	virtual void GetContextMenuActions(const UControlRigGraphNode* Node, const FGraphNodeContextMenuBuilder& Context ) = 0;
+	virtual void GetNodeContextMenuActions(const UControlRigGraphNode* Node, class UToolMenu* Menu, class UGraphNodeContextMenuContext* Context) const = 0;
 	virtual void GetContextMenuActions(const UControlRigGraphSchema* Schema, class UToolMenu* Menu, class UGraphNodeContextMenuContext* Context) const = 0;
 };
diff --git a/Engine/Plugins/FX/Niagara/Source/NiagaraEditor/Private/NiagaraNodeOutput.cpp b/Engine/Plugins/FX/Niagara/Source/NiagaraEditor/Private/NiagaraNodeOutput.cpp
--- a/Engine/Plugins/FX/Niagara/Source/NiagaraEditor/Private/NiagaraNodeOutput.cpp
+++ b/Engine/Plugins/FX/Niagara/Source/NiagaraEditor/Private/NiagaraNodeOutput.cpp
@@ -10,7 +10,7 @@
 #include "NiagaraCustomVersion.h"
 #include "ScopedTransaction.h"
 #include "EdGraphSchema_Niagara.h"
-#include "Framework/MultiBox/MultiBoxBuilder.h"
+#include "ToolMenus.h"
 #include "Widgets/Layout/SBox.h"
 #include "Widgets/Input/SEditableTextBox.h"
 #include "NiagaraEditorUtilities.h"
@@ -74,15 +74,16 @@ void UNiagaraNodeOutput::PinNameTextCommitted(const FText& Text, ETextCommit::Ty
 	}
 }
 
-void UNiagaraNodeOutput::GetContextMenuActions(const FGraphNodeContextMenuBuilder& Context) const
+void UNiagaraNodeOutput::GetNodeContextMenuActions(UToolMenu* Menu, UGraphNodeContextMenuContext* Context) const
 {
-	Super::GetContextMenuActions(Context);
-	if (Context.Pin != nullptr)
+	Super::GetNodeContextMenuActions(Menu, Context);
+
+	if (Context->Pin != nullptr)
 	{
-		Context.MenuBuilder->BeginSection("EdGraphSchema_NiagaraPinActions", LOCTEXT("EditPinMenuHeader", "Edit Pin"));
+		FToolMenuSection& Section = Menu->AddSection("EdGraphSchema_NiagaraPinActions", LOCTEXT("EditPinMenuHeader", "Edit Pin"));
 		
 		{
-			UEdGraphPin* Pin = const_cast<UEdGraphPin*>(Context.Pin);
+			UEdGraphPin* Pin = const_cast<UEdGraphPin*>(Context->Pin);
 			TSharedRef<SWidget> RenameWidget =
 				SNew(SBox)
 				.WidthOverride(100)
@@ -92,16 +93,16 @@ void UNiagaraNodeOutput::GetContextMenuActions(const FGraphNodeContextMenuBuilde
 					.Text_UObject(this, &UNiagaraNodeOutput::GetPinNameText, Pin)
 					.OnTextCommitted_UObject(const_cast<UNiagaraNodeOutput*>(this), &UNiagaraNodeOutput::PinNameTextCommitted, Pin)
 				];
-			Context.MenuBuilder->AddWidget(RenameWidget, LOCTEXT("NameMenuItem", "Name"));
+			Section.AddEntry(FToolMenuEntry::InitWidget("RenameWidget", RenameWidget, LOCTEXT("NameMenuItem", "Name")));
 		}
 
-		
 		{
-			Context.MenuBuilder->AddMenuEntry(
+			Section.AddMenuEntry(
+				"RemoveDynamicPin",
 				LOCTEXT("RemoveDynamicPin", "Remove pin"),
 				LOCTEXT("RemoveDynamicPinToolTip", "Remove this pin and any connections."),
 				FSlateIcon(),
-				FUIAction(FExecuteAction::CreateUObject(const_cast<UNiagaraNodeOutput*>(this), &UNiagaraNodeOutput::RemoveOutputPin, const_cast<UEdGraphPin*>(Context.Pin))));
+				FUIAction(FExecuteAction::CreateUObject(const_cast<UNiagaraNodeOutput*>(this), &UNiagaraNodeOutput::RemoveOutputPin, const_cast<UEdGraphPin*>(Context->Pin))));
 		}
 	}
 }
diff --git a/Engine/Plugins/FX/Niagara/Source/NiagaraEditor/Private/NiagaraNodeParameterMapBase.h b/Engine/Plugins/FX/Niagara/Source/NiagaraEditor/Private/NiagaraNodeParameterMapBase.h
--- a/Engine/Plugins/FX/Niagara/Source/NiagaraEditor/Private/NiagaraNodeParameterMapBase.h
+++ b/Engine/Plugins/FX/Niagara/Source/NiagaraEditor/Private/NiagaraNodeParameterMapBase.h
@@ -34,6 +34,8 @@ public:
 
 	virtual void CollectAddPinActions(FGraphActionListBuilderBase& OutActions, bool& bOutCreateRemainingActions, UEdGraphPin* Pin) override;
 	virtual void GetPinHoverText(const UEdGraphPin& Pin, FString& HoverTextOut) const override;
+	
+	virtual bool IncludeParentNodeContextMenu() const { return true; }
 
 protected:
 	virtual void OnPinRenamed(UEdGraphPin* RenamedPin, const FString& OldName) override;
diff --git a/Engine/Plugins/FX/Niagara/Source/NiagaraEditor/Private/NiagaraNodeParameterMapGet.cpp b/Engine/Plugins/FX/Niagara/Source/NiagaraEditor/Private/NiagaraNodeParameterMapGet.cpp
--- a/Engine/Plugins/FX/Niagara/Source/NiagaraEditor/Private/NiagaraNodeParameterMapGet.cpp
+++ b/Engine/Plugins/FX/Niagara/Source/NiagaraEditor/Private/NiagaraNodeParameterMapGet.cpp
@@ -13,7 +13,7 @@
 #include "Modules/ModuleManager.h"
 #include "Templates/SharedPointer.h"
 #include "NiagaraConstants.h"
-#include "Framework/MultiBox/MultiBoxBuilder.h"
+#include "ToolMenus.h"
 #include "Widgets/Layout/SBox.h"
 #include "Widgets/Input/SEditableTextBox.h"
 #include "EdGraph/EdGraphNode.h"
@@ -494,11 +494,11 @@ void UNiagaraNodeParameterMapGet::GetPinHoverText(const UEdGraphPin& Pin, FStrin
 	}
 }
 
-void UNiagaraNodeParameterMapGet::GetContextMenuActions(const FGraphNodeContextMenuBuilder& Context) const
+void UNiagaraNodeParameterMapGet::GetNodeContextMenuActions(UToolMenu* Menu, UGraphNodeContextMenuContext* Context) const
 {
-	UNiagaraNodeParameterMapBase::GetContextMenuActions(Context);
+	Super::GetNodeContextMenuActions(Menu, Context);
 
-	UEdGraphPin* Pin = const_cast<UEdGraphPin*>(Context.Pin);
+	UEdGraphPin* Pin = const_cast<UEdGraphPin*>(Context->Pin);
 	if (Pin && Pin->Direction == EEdGraphPinDirection::EGPD_Output)
 	{
 		
@@ -507,7 +507,7 @@ void UNiagaraNodeParameterMapGet::GetContextMenuActions(const FGraphNodeContextM
 
 		if (!FNiagaraConstants::IsNiagaraConstant(Var))
 		{
-			Context.MenuBuilder->BeginSection("EdGraphSchema_NiagaraMetaDataActions", LOCTEXT("EditPinMenuHeader", "Meta-Data"));
+			FToolMenuSection& Section = Menu->AddSection("EdGraphSchema_NiagaraMetaDataActions", LOCTEXT("EditPinMenuHeader", "Meta-Data"));
 			TSharedRef<SWidget> RenameWidget =
 				SNew(SBox)
 				.WidthOverride(100)
@@ -517,8 +517,7 @@ void UNiagaraNodeParameterMapGet::GetContextMenuActions(const FGraphNodeContextM
 					.Text_UObject(this, &UNiagaraNodeParameterMapBase::GetPinDescriptionText, Pin)
 					.OnTextCommitted_UObject(const_cast<UNiagaraNodeParameterMapGet*>(this), &UNiagaraNodeParameterMapBase::PinDescriptionTextCommitted, Pin)
 				];
-			Context.MenuBuilder->AddWidget(RenameWidget, LOCTEXT("DescMenuItem", "Description"));
-			Context.MenuBuilder->EndSection();
+			Section.AddEntry(FToolMenuEntry::InitWidget("RenameWidget", RenameWidget, LOCTEXT("DescMenuItem", "Description")));
 		}
 	}
 }
diff --git a/Engine/Plugins/FX/Niagara/Source/NiagaraEditor/Private/NiagaraNodeParameterMapGet.h b/Engine/Plugins/FX/Niagara/Source/NiagaraEditor/Private/NiagaraNodeParameterMapGet.h
--- a/Engine/Plugins/FX/Niagara/Source/NiagaraEditor/Private/NiagaraNodeParameterMapGet.h
+++ b/Engine/Plugins/FX/Niagara/Source/NiagaraEditor/Private/NiagaraNodeParameterMapGet.h
@@ -21,7 +21,8 @@ public:
 	virtual void AllocateDefaultPins() override;
 	virtual TSharedPtr<SGraphNode> CreateVisualWidget() override;
 
-	virtual void GetContextMenuActions(const FGraphNodeContextMenuBuilder& Context) const override;
+	virtual void GetNodeContextMenuActions(class UToolMenu* Menu, class UGraphNodeContextMenuContext* Context) const override;
+	virtual bool IncludeParentNodeContextMenu() const { return true; }
 	virtual bool IsPinNameEditable(const UEdGraphPin* GraphPinObj) const override;
 	virtual bool IsPinNameEditableUponCreation(const UEdGraphPin* GraphPinObj) const override;
 	virtual bool VerifyEditablePinName(const FText& InName, FText& OutErrorMessage, const UEdGraphPin* InGraphPinObj) const override;
diff --git a/Engine/Plugins/FX/Niagara/Source/NiagaraEditor/Private/NiagaraNodeParameterMapSet.cpp b/Engine/Plugins/FX/Niagara/Source/NiagaraEditor/Private/NiagaraNodeParameterMapSet.cpp
--- a/Engine/Plugins/FX/Niagara/Source/NiagaraEditor/Private/NiagaraNodeParameterMapSet.cpp
+++ b/Engine/Plugins/FX/Niagara/Source/NiagaraEditor/Private/NiagaraNodeParameterMapSet.cpp
@@ -8,7 +8,7 @@
 #include "Templates/SharedPointer.h"
 #include "NiagaraGraph.h"
 #include "NiagaraConstants.h"
-#include "Framework/MultiBox/MultiBoxBuilder.h"
+#include "ToolMenus.h"
 #include "Widgets/Layout/SBox.h"
 #include "Widgets/Input/SEditableTextBox.h"
 #include "EdGraph/EdGraphNode.h"
@@ -251,11 +251,11 @@ void UNiagaraNodeParameterMapSet::BuildParameterMapHistory(FNiagaraParameterMapH
 	OutHistory.RegisterParameterMapPin(ParamMapIdx, GetOutputPin(0));
 }
 
-void UNiagaraNodeParameterMapSet::GetContextMenuActions(const FGraphNodeContextMenuBuilder& Context) const
+void UNiagaraNodeParameterMapSet::GetNodeContextMenuActions(UToolMenu* Menu, UGraphNodeContextMenuContext* Context) const
 {
-	UNiagaraNodeParameterMapBase::GetContextMenuActions(Context);
+	Super::GetNodeContextMenuActions(Menu, Context);
 
-	UEdGraphPin* Pin = const_cast<UEdGraphPin*>(Context.Pin);
+	UEdGraphPin* Pin = const_cast<UEdGraphPin*>(Context->Pin);
 	if (Pin && Pin->Direction == EEdGraphPinDirection::EGPD_Input)
 	{
 		FNiagaraVariable Var = CastChecked<UEdGraphSchema_Niagara>(GetSchema())->PinToNiagaraVariable(Pin);
@@ -263,7 +263,7 @@ void UNiagaraNodeParameterMapSet::GetContextMenuActions(const FGraphNodeContextM
 
 		if (!FNiagaraConstants::IsNiagaraConstant(Var))
 		{
-			Context.MenuBuilder->BeginSection("EdGraphSchema_NiagaraMetaDataActions", LOCTEXT("EditPinMenuHeader", "Meta-Data"));
+			FToolMenuSection& Section = Menu->AddSection("EdGraphSchema_NiagaraMetaDataActions", LOCTEXT("EditPinMenuHeader", "Meta-Data"));
 			TSharedRef<SWidget> RenameWidget =
 				SNew(SBox)
 				.WidthOverride(100)
@@ -273,9 +273,7 @@ void UNiagaraNodeParameterMapSet::GetContextMenuActions(const FGraphNodeContextM
 					.Text_UObject(this, &UNiagaraNodeParameterMapBase::GetPinDescriptionText, Pin)
 					.OnTextCommitted_UObject(const_cast<UNiagaraNodeParameterMapSet*>(this), &UNiagaraNodeParameterMapBase::PinDescriptionTextCommitted, Pin)
 				];
-			Context.MenuBuilder->AddWidget(RenameWidget, LOCTEXT("DescMenuItem", "Description"));
-
-			Context.MenuBuilder->EndSection();
+			Section.AddEntry(FToolMenuEntry::InitWidget("RenameWidget", RenameWidget, LOCTEXT("DescMenuItem", "Description")));
 		}
 	}
 }
diff --git a/Engine/Plugins/FX/Niagara/Source/NiagaraEditor/Private/NiagaraNodeParameterMapSet.h b/Engine/Plugins/FX/Niagara/Source/NiagaraEditor/Private/NiagaraNodeParameterMapSet.h
--- a/Engine/Plugins/FX/Niagara/Source/NiagaraEditor/Private/NiagaraNodeParameterMapSet.h
+++ b/Engine/Plugins/FX/Niagara/Source/NiagaraEditor/Private/NiagaraNodeParameterMapSet.h
@@ -32,7 +32,8 @@ public:
 
 	virtual FText GetNodeTitle(ENodeTitleType::Type TitleType) const;
 
-	virtual void GetContextMenuActions(const FGraphNodeContextMenuBuilder& Context) const override;
+	virtual void GetNodeContextMenuActions(class UToolMenu* Menu, class UGraphNodeContextMenuContext* Context) const override;
+	virtual bool IncludeParentNodeContextMenu() const { return true; }
 	virtual void PostLoad() override;
 
 	void SetPinName(UEdGraphPin* InPin, const FName& InName);
diff --git a/Engine/Plugins/FX/Niagara/Source/NiagaraEditor/Private/NiagaraNodeWithDynamicPins.cpp b/Engine/Plugins/FX/Niagara/Source/NiagaraEditor/Private/NiagaraNodeWithDynamicPins.cpp
--- a/Engine/Plugins/FX/Niagara/Source/NiagaraEditor/Private/NiagaraNodeWithDynamicPins.cpp
+++ b/Engine/Plugins/FX/Niagara/Source/NiagaraEditor/Private/NiagaraNodeWithDynamicPins.cpp
@@ -6,7 +6,7 @@
 #include "NiagaraGraph.h"
 #include "Framework/Commands/UIAction.h"
 #include "ScopedTransaction.h"
-#include "Framework/MultiBox/MultiBoxBuilder.h"
+#include "ToolMenus.h"
 #include "Widgets/Input/SEditableTextBox.h"
 #include "Widgets/Layout/SBox.h"
 #include "SNiagaraGraphPinAdd.h"
@@ -161,15 +161,15 @@ void UNiagaraNodeWithDynamicPins::MoveDynamicPin(UEdGraphPin* Pin, int32 Directi
 	}
 }
 
-void UNiagaraNodeWithDynamicPins::GetContextMenuActions(const FGraphNodeContextMenuBuilder& Context) const
+void UNiagaraNodeWithDynamicPins::GetNodeContextMenuActions(UToolMenu* Menu, UGraphNodeContextMenuContext* Context) const
 {
-	Super::GetContextMenuActions(Context);
-	if (Context.Pin != nullptr)
+	Super::GetNodeContextMenuActions(Menu, Context);
+	if (Context->Pin != nullptr)
 	{
-		Context.MenuBuilder->BeginSection("EdGraphSchema_NiagaraPinActions", LOCTEXT("EditPinMenuHeader", "Edit Pin"));
-		if (CanRenamePin(Context.Pin))
+		FToolMenuSection& Section = Menu->AddSection("EdGraphSchema_NiagaraPinActions", LOCTEXT("EditPinMenuHeader", "Edit Pin"));
+		if (CanRenamePin(Context->Pin))
 		{
-			UEdGraphPin* Pin = const_cast<UEdGraphPin*>(Context.Pin);
+			UEdGraphPin* Pin = const_cast<UEdGraphPin*>(Context->Pin);
 			TSharedRef<SWidget> RenameWidget =
 				SNew(SBox)
 				.WidthOverride(100)
@@ -179,20 +179,21 @@ void UNiagaraNodeWithDynamicPins::GetContextMenuActions(const FGraphNodeContextM
 					.Text_UObject(this, &UNiagaraNodeWithDynamicPins::GetPinNameText, Pin)
 					.OnTextCommitted_UObject(const_cast<UNiagaraNodeWithDynamicPins*>(this), &UNiagaraNodeWithDynamicPins::PinNameTextCommitted, Pin)
 				];
-			Context.MenuBuilder->AddWidget(RenameWidget, LOCTEXT("NameMenuItem", "Name"));
+			Section.AddEntry(FToolMenuEntry::InitWidget("RenameWidget", RenameWidget, LOCTEXT("NameMenuItem", "Name")));
 		}
-		if (CanRemovePin(Context.Pin))
+		if (CanRemovePin(Context->Pin))
 		{
-			Context.MenuBuilder->AddMenuEntry(
+			Section.AddMenuEntry(
+				"RemoveDynamicPin",
 				LOCTEXT("RemoveDynamicPin", "Remove pin"),
 				LOCTEXT("RemoveDynamicPinToolTip", "Remove this pin and any connections."),
 				FSlateIcon(),
-				FUIAction(FExecuteAction::CreateUObject(const_cast<UNiagaraNodeWithDynamicPins*>(this), &UNiagaraNodeWithDynamicPins::RemoveDynamicPinFromMenu, const_cast<UEdGraphPin*>(Context.Pin))));
+				FUIAction(FExecuteAction::CreateUObject(const_cast<UNiagaraNodeWithDynamicPins*>(this), &UNiagaraNodeWithDynamicPins::RemoveDynamicPinFromMenu, const_cast<UEdGraphPin*>(Context->Pin))));
 		}
-		if (CanMovePin(Context.Pin))
+		if (CanMovePin(Context->Pin))
 		{
 			TArray<UEdGraphPin*> SameDirectionPins;
-			if (Context.Pin->Direction == EEdGraphPinDirection::EGPD_Input)
+			if (Context->Pin->Direction == EEdGraphPinDirection::EGPD_Input)
 			{
 				GetInputPins(SameDirectionPins);
 			}
@@ -201,26 +202,27 @@ void UNiagaraNodeWithDynamicPins::GetContextMenuActions(const FGraphNodeContextM
 				GetOutputPins(SameDirectionPins);
 			}
 			int32 PinIdx = INDEX_NONE;
-			SameDirectionPins.Find(const_cast<UEdGraphPin*>(Context.Pin), PinIdx);
+			SameDirectionPins.Find(const_cast<UEdGraphPin*>(Context->Pin), PinIdx);
 
 			if (PinIdx != 0)
 			{
-				Context.MenuBuilder->AddMenuEntry(
+				Section.AddMenuEntry(
+					"MoveDynamicPinUp",
 					LOCTEXT("MoveDynamicPinUp", "Move pin up"),
 					LOCTEXT("MoveDynamicPinToolTipUp", "Move this pin and any connections one slot up."),
 					FSlateIcon(),
-					FUIAction(FExecuteAction::CreateUObject(const_cast<UNiagaraNodeWithDynamicPins*>(this), &UNiagaraNodeWithDynamicPins::MoveDynamicPinFromMenu, const_cast<UEdGraphPin*>(Context.Pin), -1)));
+					FUIAction(FExecuteAction::CreateUObject(const_cast<UNiagaraNodeWithDynamicPins*>(this), &UNiagaraNodeWithDynamicPins::MoveDynamicPinFromMenu, const_cast<UEdGraphPin*>(Context->Pin), -1)));
 			}
 			if (PinIdx >= 0 && PinIdx < SameDirectionPins.Num() - 1)
 			{
-				Context.MenuBuilder->AddMenuEntry(
+				Section.AddMenuEntry(
+					"MoveDynamicPinDown",
 					LOCTEXT("MoveDynamicPinDown", "Move pin down"),
 					LOCTEXT("MoveDynamicPinToolTipDown", "Move this pin and any connections one slot down."),
 					FSlateIcon(),
-					FUIAction(FExecuteAction::CreateUObject(const_cast<UNiagaraNodeWithDynamicPins*>(this), &UNiagaraNodeWithDynamicPins::MoveDynamicPin, const_cast<UEdGraphPin*>(Context.Pin), 1)));
+					FUIAction(FExecuteAction::CreateUObject(const_cast<UNiagaraNodeWithDynamicPins*>(this), &UNiagaraNodeWithDynamicPins::MoveDynamicPin, const_cast<UEdGraphPin*>(Context->Pin), 1)));
 			}
 		}
-		Context.MenuBuilder->EndSection();
 	}
 }
 
diff --git a/Engine/Plugins/FX/Niagara/Source/NiagaraEditor/Public/NiagaraNodeOutput.h b/Engine/Plugins/FX/Niagara/Source/NiagaraEditor/Public/NiagaraNodeOutput.h
--- a/Engine/Plugins/FX/Niagara/Source/NiagaraEditor/Public/NiagaraNodeOutput.h
+++ b/Engine/Plugins/FX/Niagara/Source/NiagaraEditor/Public/NiagaraNodeOutput.h
@@ -41,7 +41,8 @@ public:
 	virtual bool CanDuplicateNode() const override;
 	virtual FText GetNodeTitle(ENodeTitleType::Type TitleType) const override;
 	virtual FLinearColor GetNodeTitleColor() const override;
-	virtual void GetContextMenuActions(const FGraphNodeContextMenuBuilder& Context) const override;
+	virtual void GetNodeContextMenuActions(class UToolMenu* Menu, class UGraphNodeContextMenuContext* Context) const override;
+	virtual bool IncludeParentNodeContextMenu() const { return true; }
 	//~ End EdGraphNode Interface
 
 	/** Notifies the node that it's output variables have been modified externally. */
diff --git a/Engine/Plugins/FX/Niagara/Source/NiagaraEditor/Public/NiagaraNodeWithDynamicPins.h b/Engine/Plugins/FX/Niagara/Source/NiagaraEditor/Public/NiagaraNodeWithDynamicPins.h
--- a/Engine/Plugins/FX/Niagara/Source/NiagaraEditor/Public/NiagaraNodeWithDynamicPins.h
+++ b/Engine/Plugins/FX/Niagara/Source/NiagaraEditor/Public/NiagaraNodeWithDynamicPins.h
@@ -23,7 +23,8 @@ public:
 
 	//~ UEdGraphNode interface
 	virtual void PinConnectionListChanged(UEdGraphPin* Pin) override;
-	virtual void GetContextMenuActions(const FGraphNodeContextMenuBuilder& Context) const override;
+	virtual void GetNodeContextMenuActions(class UToolMenu* Menu, class UGraphNodeContextMenuContext* Context) const override;
+	virtual bool IncludeParentNodeContextMenu() const { return true; }
 
 	/** Requests a new pin be added to the node with the specified direction, type, and name. */
 	UEdGraphPin* RequestNewTypedPin(EEdGraphPinDirection Direction, const FNiagaraTypeDefinition& Type, FName InName);
diff --git a/Engine/Source/Developer/ToolMenus/Private/ToolMenuSection.cpp b/Engine/Source/Developer/ToolMenus/Private/ToolMenuSection.cpp
--- a/Engine/Source/Developer/ToolMenus/Private/ToolMenuSection.cpp
+++ b/Engine/Source/Developer/ToolMenus/Private/ToolMenuSection.cpp
@@ -108,6 +108,11 @@ FToolMenuEntry& FToolMenuSection::AddSubMenu(const FName InParentMenu, const FNa
 	return AddEntry(FToolMenuEntry::InitSubMenu(InParentMenu, InName, InLabel, InToolTip, InMakeMenu, bInOpenSubMenuOnClick, InIcon, bShouldCloseWindowAfterMenuSelection));
 }
 
+FToolMenuEntry& FToolMenuSection::AddSubMenu(const FName InName, const TAttribute<FText>& InLabel, const TAttribute<FText>& InToolTip, const FNewToolMenuChoice& InMakeMenu, bool bInOpenSubMenuOnClick, const TAttribute<FSlateIcon>& InIcon, const bool bShouldCloseWindowAfterMenuSelection)
+{
+	return AddEntry(FToolMenuEntry::InitSubMenu(NAME_None, InName, InLabel, InToolTip, InMakeMenu, bInOpenSubMenuOnClick, InIcon, bShouldCloseWindowAfterMenuSelection));
+}
+
 int32 FToolMenuSection::IndexOfBlock(const FName InName) const
 {
 	for (int32 i=0; i < Blocks.Num(); ++i)
diff --git a/Engine/Source/Developer/ToolMenus/Public/ToolMenuSection.h b/Engine/Source/Developer/ToolMenus/Public/ToolMenuSection.h
--- a/Engine/Source/Developer/ToolMenus/Public/ToolMenuSection.h
+++ b/Engine/Source/Developer/ToolMenus/Public/ToolMenuSection.h
@@ -43,6 +43,7 @@ public:
 	FToolMenuEntry& AddMenuSeparator(const FName InName);
 
 	FToolMenuEntry& AddSubMenu(const FName InParentMenu, const FName InName, const TAttribute<FText>& InLabel, const TAttribute<FText>& InToolTip, const FNewToolMenuChoice& InMakeMenu, bool bInOpenSubMenuOnClick = false, const TAttribute<FSlateIcon>& InIcon = TAttribute<FSlateIcon>(), const bool bShouldCloseWindowAfterMenuSelection = true);
+	FToolMenuEntry& AddSubMenu(const FName InName, const TAttribute<FText>& InLabel, const TAttribute<FText>& InToolTip, const FNewToolMenuChoice& InMakeMenu, bool bInOpenSubMenuOnClick = false, const TAttribute<FSlateIcon>& InIcon = TAttribute<FSlateIcon>(), const bool bShouldCloseWindowAfterMenuSelection = true);
 
 	template <typename TContextType>
 	TContextType* FindContext() const
diff --git a/Engine/Source/Editor/AnimGraph/Classes/AnimGraphNode_AimOffsetLookAt.h b/Engine/Source/Editor/AnimGraph/Classes/AnimGraphNode_AimOffsetLookAt.h
--- a/Engine/Source/Editor/AnimGraph/Classes/AnimGraphNode_AimOffsetLookAt.h
+++ b/Engine/Source/Editor/AnimGraph/Classes/AnimGraphNode_AimOffsetLookAt.h
@@ -36,7 +36,7 @@ class UAnimGraphNode_AimOffsetLookAt : public UAnimGraphNode_BlendSpaceBase
 	// End of UAnimGraphNode_Base interface
 
 	// UK2Node interface
-	virtual void GetContextMenuActions(const FGraphNodeContextMenuBuilder& Context) const override;
+	virtual void GetNodeContextMenuActions(class UToolMenu* Menu, class UGraphNodeContextMenuContext* Context) const override;
 	virtual void GetMenuActions(FBlueprintActionDatabaseRegistrar& ActionRegistrar) const override;
 	virtual FBlueprintNodeSignature GetSignature() const override;
 	// End of UK2Node interface
diff --git a/Engine/Source/Editor/AnimGraph/Classes/AnimGraphNode_BlendListByEnum.h b/Engine/Source/Editor/AnimGraph/Classes/AnimGraphNode_BlendListByEnum.h
--- a/Engine/Source/Editor/AnimGraph/Classes/AnimGraphNode_BlendListByEnum.h
+++ b/Engine/Source/Editor/AnimGraph/Classes/AnimGraphNode_BlendListByEnum.h
@@ -35,7 +35,7 @@ public:
 	// End of UEdGraphNode interface
 
 	// UK2Node interface
-	virtual void GetContextMenuActions(const FGraphNodeContextMenuBuilder& Context) const override;
+	virtual void GetNodeContextMenuActions(class UToolMenu* Menu, class UGraphNodeContextMenuContext* Context) const override;
 	// End of UK2Node interface
 
 	// UAnimGraphNode_Base interface
diff --git a/Engine/Source/Editor/AnimGraph/Classes/AnimGraphNode_BlendListByInt.h b/Engine/Source/Editor/AnimGraph/Classes/AnimGraphNode_BlendListByInt.h
--- a/Engine/Source/Editor/AnimGraph/Classes/AnimGraphNode_BlendListByInt.h
+++ b/Engine/Source/Editor/AnimGraph/Classes/AnimGraphNode_BlendListByInt.h
@@ -28,6 +28,6 @@ class UAnimGraphNode_BlendListByInt : public UAnimGraphNode_BlendListBase
 	// End of UEdGraphNode interface
 
 	// UK2Node interface
-	virtual void GetContextMenuActions(const FGraphNodeContextMenuBuilder& Context) const override;
+	virtual void GetNodeContextMenuActions(class UToolMenu* Menu, class UGraphNodeContextMenuContext* Context) const override;
 	// End of UK2Node interface
 };
diff --git a/Engine/Source/Editor/AnimGraph/Classes/AnimGraphNode_BlendSpaceEvaluator.h b/Engine/Source/Editor/AnimGraph/Classes/AnimGraphNode_BlendSpaceEvaluator.h
--- a/Engine/Source/Editor/AnimGraph/Classes/AnimGraphNode_BlendSpaceEvaluator.h
+++ b/Engine/Source/Editor/AnimGraph/Classes/AnimGraphNode_BlendSpaceEvaluator.h
@@ -37,7 +37,7 @@ class UAnimGraphNode_BlendSpaceEvaluator : public UAnimGraphNode_BlendSpaceBase
 	// End of UAnimGraphNode_Base interface
 
 	// UK2Node interface
-	virtual void GetContextMenuActions(const FGraphNodeContextMenuBuilder& Context) const override;
+	virtual void GetNodeContextMenuActions(class UToolMenu* Menu, class UGraphNodeContextMenuContext* Context) const override;
 	// End of UK2Node interface
 
 	// UAnimGraphNode_AssetPlayerBase interface
diff --git a/Engine/Source/Editor/AnimGraph/Classes/AnimGraphNode_BlendSpacePlayer.h b/Engine/Source/Editor/AnimGraph/Classes/AnimGraphNode_BlendSpacePlayer.h
--- a/Engine/Source/Editor/AnimGraph/Classes/AnimGraphNode_BlendSpacePlayer.h
+++ b/Engine/Source/Editor/AnimGraph/Classes/AnimGraphNode_BlendSpacePlayer.h
@@ -38,7 +38,7 @@ class UAnimGraphNode_BlendSpacePlayer : public UAnimGraphNode_BlendSpaceBase
 	// End of UAnimGraphNode_Base interface
 
 	// UK2Node interface
-	virtual void GetContextMenuActions(const FGraphNodeContextMenuBuilder& Context) const override;
+	virtual void GetNodeContextMenuActions(class UToolMenu* Menu, class UGraphNodeContextMenuContext* Context) const override;
 	virtual void GetMenuActions(FBlueprintActionDatabaseRegistrar& ActionRegistrar) const override;
 	virtual FBlueprintNodeSignature GetSignature() const override;
 	// End of UK2Node interface
diff --git a/Engine/Source/Editor/AnimGraph/Classes/AnimGraphNode_LayeredBoneBlend.h b/Engine/Source/Editor/AnimGraph/Classes/AnimGraphNode_LayeredBoneBlend.h
--- a/Engine/Source/Editor/AnimGraph/Classes/AnimGraphNode_LayeredBoneBlend.h
+++ b/Engine/Source/Editor/AnimGraph/Classes/AnimGraphNode_LayeredBoneBlend.h
@@ -32,7 +32,7 @@ class UAnimGraphNode_LayeredBoneBlend : public UAnimGraphNode_BlendListBase
 	// End of UEdGraphNode interface
 
 	// UK2Node interface
-	virtual void GetContextMenuActions(const FGraphNodeContextMenuBuilder& Context) const override;
+	virtual void GetNodeContextMenuActions(class UToolMenu* Menu, class UGraphNodeContextMenuContext* Context) const override;
 	// End of UK2Node interface
 
 	// UAnimGraphNode_Base interface
diff --git a/Engine/Source/Editor/AnimGraph/Classes/AnimGraphNode_ModifyCurve.h b/Engine/Source/Editor/AnimGraph/Classes/AnimGraphNode_ModifyCurve.h
--- a/Engine/Source/Editor/AnimGraph/Classes/AnimGraphNode_ModifyCurve.h
+++ b/Engine/Source/Editor/AnimGraph/Classes/AnimGraphNode_ModifyCurve.h
@@ -33,7 +33,7 @@ public:
 	// End of UAnimGraphNode_Base interface
 
 	// UK2Node interface
-	virtual void GetContextMenuActions(const FGraphNodeContextMenuBuilder& Context) const override;
+	virtual void GetNodeContextMenuActions(class UToolMenu* Menu, class UGraphNodeContextMenuContext* Context) const override;
 	// End of UK2Node interface
 
 private:
diff --git a/Engine/Source/Editor/AnimGraph/Classes/AnimGraphNode_MultiWayBlend.h b/Engine/Source/Editor/AnimGraph/Classes/AnimGraphNode_MultiWayBlend.h
--- a/Engine/Source/Editor/AnimGraph/Classes/AnimGraphNode_MultiWayBlend.h
+++ b/Engine/Source/Editor/AnimGraph/Classes/AnimGraphNode_MultiWayBlend.h
@@ -31,7 +31,7 @@ class UAnimGraphNode_MultiWayBlend : public UAnimGraphNode_Base
 	//~ End UAnimGraphNode_Base Interface
 
 	// UK2Node interface
-	virtual void GetContextMenuActions(const FGraphNodeContextMenuBuilder& Context) const override;
+	virtual void GetNodeContextMenuActions(class UToolMenu* Menu, class UGraphNodeContextMenuContext* Context) const override;
 	// End of UK2Node interface
 
 private:
diff --git a/Engine/Source/Editor/AnimGraph/Classes/AnimGraphNode_PoseBlendNode.h b/Engine/Source/Editor/AnimGraph/Classes/AnimGraphNode_PoseBlendNode.h
--- a/Engine/Source/Editor/AnimGraph/Classes/AnimGraphNode_PoseBlendNode.h
+++ b/Engine/Source/Editor/AnimGraph/Classes/AnimGraphNode_PoseBlendNode.h
@@ -35,7 +35,7 @@ class UAnimGraphNode_PoseBlendNode : public UAnimGraphNode_PoseHandler
 	// End of UAnimGraphNode_Base
 
 	// UK2Node interface
-	virtual void GetContextMenuActions(const FGraphNodeContextMenuBuilder& Context) const override;
+	virtual void GetNodeContextMenuActions(class UToolMenu* Menu, class UGraphNodeContextMenuContext* Context) const override;
 	// End of UK2Node interface
 
 private:
diff --git a/Engine/Source/Editor/AnimGraph/Classes/AnimGraphNode_PoseByName.h b/Engine/Source/Editor/AnimGraph/Classes/AnimGraphNode_PoseByName.h
--- a/Engine/Source/Editor/AnimGraph/Classes/AnimGraphNode_PoseByName.h
+++ b/Engine/Source/Editor/AnimGraph/Classes/AnimGraphNode_PoseByName.h
@@ -37,7 +37,7 @@ class UAnimGraphNode_PoseByName : public UAnimGraphNode_AssetPlayerBase
 	// End of UAnimGraphNode_Base
 
 	// UK2Node interface
-	virtual void GetContextMenuActions(const FGraphNodeContextMenuBuilder& Context) const override;
+	virtual void GetNodeContextMenuActions(class UToolMenu* Menu, class UGraphNodeContextMenuContext* Context) const override;
 	virtual void GetMenuActions(FBlueprintActionDatabaseRegistrar& ActionRegistrar) const override;
 	//virtual void GetMenuActions(FBlueprintActionDatabaseRegistrar& ActionRegistrar) const override;
 	// End of UK2Node interface
diff --git a/Engine/Source/Editor/AnimGraph/Classes/AnimGraphNode_RotationOffsetBlendSpace.h b/Engine/Source/Editor/AnimGraph/Classes/AnimGraphNode_RotationOffsetBlendSpace.h
--- a/Engine/Source/Editor/AnimGraph/Classes/AnimGraphNode_RotationOffsetBlendSpace.h
+++ b/Engine/Source/Editor/AnimGraph/Classes/AnimGraphNode_RotationOffsetBlendSpace.h
@@ -37,7 +37,7 @@ class UAnimGraphNode_RotationOffsetBlendSpace: public UAnimGraphNode_BlendSpaceB
 	// End of UAnimGraphNode_Base interface
 
 	// UK2Node interface
-	virtual void GetContextMenuActions(const FGraphNodeContextMenuBuilder& Context) const override;
+	virtual void GetNodeContextMenuActions(class UToolMenu* Menu, class UGraphNodeContextMenuContext* Context) const override;
 	virtual void GetMenuActions(FBlueprintActionDatabaseRegistrar& ActionRegistrar) const override;
 	virtual FBlueprintNodeSignature GetSignature() const override;
 	// End of UK2Node interface
diff --git a/Engine/Source/Editor/AnimGraph/Classes/AnimGraphNode_SequenceEvaluator.h b/Engine/Source/Editor/AnimGraph/Classes/AnimGraphNode_SequenceEvaluator.h
--- a/Engine/Source/Editor/AnimGraph/Classes/AnimGraphNode_SequenceEvaluator.h
+++ b/Engine/Source/Editor/AnimGraph/Classes/AnimGraphNode_SequenceEvaluator.h
@@ -39,7 +39,7 @@ class UAnimGraphNode_SequenceEvaluator : public UAnimGraphNode_AssetPlayerBase
 	// End of UAnimGraphNode_Base
 
 	// UK2Node interface
-	virtual void GetContextMenuActions(const FGraphNodeContextMenuBuilder& Context) const override;
+	virtual void GetNodeContextMenuActions(class UToolMenu* Menu, class UGraphNodeContextMenuContext* Context) const override;
 	virtual void GetMenuActions(FBlueprintActionDatabaseRegistrar& ActionRegistrar) const override;
 	// End of UK2Node interface
 
diff --git a/Engine/Source/Editor/AnimGraph/Classes/AnimGraphNode_SequencePlayer.h b/Engine/Source/Editor/AnimGraph/Classes/AnimGraphNode_SequencePlayer.h
--- a/Engine/Source/Editor/AnimGraph/Classes/AnimGraphNode_SequencePlayer.h
+++ b/Engine/Source/Editor/AnimGraph/Classes/AnimGraphNode_SequencePlayer.h
@@ -44,7 +44,7 @@ class UAnimGraphNode_SequencePlayer : public UAnimGraphNode_AssetPlayerBase
 	// End of UAnimGraphNode_Base interface
 
 	// UK2Node interface
-	virtual void GetContextMenuActions(const FGraphNodeContextMenuBuilder& Context) const override;
+	virtual void GetNodeContextMenuActions(class UToolMenu* Menu, class UGraphNodeContextMenuContext* Context) const override;
 	// End of UK2Node interface
 
 	// UAnimGraphNode_AssetPlayerBase interface
diff --git a/Engine/Source/Editor/AnimGraph/Private/AnimGraphNode_AimOffsetLookAt.cpp b/Engine/Source/Editor/AnimGraph/Private/AnimGraphNode_AimOffsetLookAt.cpp
--- a/Engine/Source/Editor/AnimGraph/Private/AnimGraphNode_AimOffsetLookAt.cpp
+++ b/Engine/Source/Editor/AnimGraph/Private/AnimGraphNode_AimOffsetLookAt.cpp
@@ -10,7 +10,7 @@
 #include "Animation/AimOffsetBlendSpace1D.h"
 #include "UObject/UObjectHash.h"
 #include "UObject/UObjectIterator.h"
-#include "Framework/MultiBox/MultiBoxBuilder.h"
+#include "ToolMenus.h"
 
 /////////////////////////////////////////////////////
 // UAnimGraphNode_RotationOffsetBlendSpace
@@ -208,17 +208,16 @@ void UAnimGraphNode_AimOffsetLookAt::ValidateAnimNodeDuringCompilation(class USk
 	}
 }
 
-void UAnimGraphNode_AimOffsetLookAt::GetContextMenuActions(const FGraphNodeContextMenuBuilder& Context) const
+void UAnimGraphNode_AimOffsetLookAt::GetNodeContextMenuActions(UToolMenu* Menu, UGraphNodeContextMenuContext* Context) const
 {
-	if (!Context.bIsDebugging)
+	if (!Context->bIsDebugging)
 	{
 		// add an option to convert to single frame
-		Context.MenuBuilder->BeginSection("AnimGraphNodeBlendSpacePlayer", NSLOCTEXT("A3Nodes", "BlendSpaceHeading", "Blend Space"));
 		{
-			Context.MenuBuilder->AddMenuEntry(FGraphEditorCommands::Get().OpenRelatedAsset);
-			Context.MenuBuilder->AddMenuEntry(FGraphEditorCommands::Get().ConvertToAimOffsetSimple);
+			FToolMenuSection& Section = Menu->AddSection("AnimGraphNodeBlendSpacePlayer", NSLOCTEXT("A3Nodes", "BlendSpaceHeading", "Blend Space"));
+			Section.AddMenuEntry(FGraphEditorCommands::Get().OpenRelatedAsset);
+			Section.AddMenuEntry(FGraphEditorCommands::Get().ConvertToAimOffsetSimple);
 		}
-		Context.MenuBuilder->EndSection();
 	}
 }
 
diff --git a/Engine/Source/Editor/AnimGraph/Private/AnimGraphNode_BlendListByEnum.cpp b/Engine/Source/Editor/AnimGraph/Private/AnimGraphNode_BlendListByEnum.cpp
--- a/Engine/Source/Editor/AnimGraph/Private/AnimGraphNode_BlendListByEnum.cpp
+++ b/Engine/Source/Editor/AnimGraph/Private/AnimGraphNode_BlendListByEnum.cpp
@@ -3,7 +3,7 @@
 #include "AnimGraphNode_BlendListByEnum.h"
 #include "Textures/SlateIcon.h"
 #include "Framework/Commands/UIAction.h"
-#include "Framework/MultiBox/MultiBoxBuilder.h"
+#include "ToolMenus.h"
 #include "Kismet2/BlueprintEditorUtils.h"
 
 #include "ScopedTransaction.h"
@@ -77,16 +77,16 @@ void UAnimGraphNode_BlendListByEnum::GetMenuActions(FBlueprintActionDatabaseRegi
 	}) );
 }
 
-void UAnimGraphNode_BlendListByEnum::GetContextMenuActions(const FGraphNodeContextMenuBuilder& Context) const
+void UAnimGraphNode_BlendListByEnum::GetNodeContextMenuActions(UToolMenu* Menu, UGraphNodeContextMenuContext* Context) const
 {
-	if (!Context.bIsDebugging && BoundEnum)
+	if (!Context->bIsDebugging && BoundEnum)
 	{
-		if (Context.Pin && (Context.Pin->Direction == EGPD_Input))
+		if (Context->Pin && (Context->Pin->Direction == EGPD_Input))
 		{
 			int32 RawArrayIndex = 0;
 			bool bIsPosePin = false;
 			bool bIsTimePin = false;
-			GetPinInformation(Context.Pin->PinName.ToString(), /*out*/ RawArrayIndex, /*out*/ bIsPosePin, /*out*/ bIsTimePin);
+			GetPinInformation(Context->Pin->PinName.ToString(), /*out*/ RawArrayIndex, /*out*/ bIsPosePin, /*out*/ bIsTimePin);
 
 			if (bIsPosePin || bIsTimePin)
 			{
@@ -95,8 +95,9 @@ void UAnimGraphNode_BlendListByEnum::GetContextMenuActions(const FGraphNodeConte
 				if (ExposedEnumIndex != INDEX_NONE)
 				{
 					// Offer to remove this specific pin
-					FUIAction Action = FUIAction( FExecuteAction::CreateUObject( const_cast<UAnimGraphNode_BlendListByEnum*>(this), &UAnimGraphNode_BlendListByEnum::RemovePinFromBlendList, const_cast<UEdGraphPin*>(Context.Pin)) );
-					Context.MenuBuilder->AddMenuEntry( LOCTEXT("RemovePose", "Remove Pose"), FText::GetEmpty(), FSlateIcon(), Action );
+					FUIAction Action = FUIAction( FExecuteAction::CreateUObject( const_cast<UAnimGraphNode_BlendListByEnum*>(this), &UAnimGraphNode_BlendListByEnum::RemovePinFromBlendList, const_cast<UEdGraphPin*>(Context->Pin)) );
+					FToolMenuSection& Section = Menu->AddSection("RemovePose");
+					Section.AddMenuEntry("RemovePose", LOCTEXT("RemovePose", "Remove Pose"), FText::GetEmpty(), FSlateIcon(), Action);
 				}
 			}
 		}
@@ -115,17 +116,17 @@ void UAnimGraphNode_BlendListByEnum::GetContextMenuActions(const FGraphNodeConte
 				if (!bAddedHeader)
 				{
 					bAddedHeader = true;
-					Context.MenuBuilder->BeginSection("AnimGraphNodeAddElementPin", LOCTEXT("ExposeHeader", "Add pin for element"));
 					{
+						FToolMenuSection& Section = Menu->AddSection("AnimGraphNodeAddElementPin", LOCTEXT("ExposeHeader", "Add pin for element"));
 						FUIAction Action = FUIAction( FExecuteAction::CreateUObject( const_cast<UAnimGraphNode_BlendListByEnum*>(this), &UAnimGraphNode_BlendListByEnum::ExposeEnumElementAsPin, ElementName) );
-						Context.MenuBuilder->AddMenuEntry(PrettyElementName, PrettyElementName, FSlateIcon(), Action);
+						Section.AddMenuEntry("PrettyElementName", PrettyElementName, PrettyElementName, FSlateIcon(), Action);
 					}
-					Context.MenuBuilder->EndSection();
 				}
 				else
 				{
 					FUIAction Action = FUIAction( FExecuteAction::CreateUObject( const_cast<UAnimGraphNode_BlendListByEnum*>(this), &UAnimGraphNode_BlendListByEnum::ExposeEnumElementAsPin, ElementName) );
-					Context.MenuBuilder->AddMenuEntry(PrettyElementName, PrettyElementName, FSlateIcon(), Action);
+					FToolMenuSection& Section = Menu->AddSection("PrettyElementName");
+					Section.AddMenuEntry("PrettyElementName", PrettyElementName, PrettyElementName, FSlateIcon(), Action);
 				}
 			}
 		}
diff --git a/Engine/Source/Editor/AnimGraph/Private/AnimGraphNode_BlendListByInt.cpp b/Engine/Source/Editor/AnimGraph/Private/AnimGraphNode_BlendListByInt.cpp
--- a/Engine/Source/Editor/AnimGraph/Private/AnimGraphNode_BlendListByInt.cpp
+++ b/Engine/Source/Editor/AnimGraph/Private/AnimGraphNode_BlendListByInt.cpp
@@ -2,7 +2,7 @@
 
 
 #include "AnimGraphNode_BlendListByInt.h"
-#include "Framework/MultiBox/MultiBoxBuilder.h"
+#include "ToolMenus.h"
 #include "Kismet2/BlueprintEditorUtils.h"
 
 #include "GraphEditorActions.h"
@@ -69,26 +69,25 @@ void UAnimGraphNode_BlendListByInt::RemovePinFromBlendList(UEdGraphPin* Pin)
 	}
 }
 
-void UAnimGraphNode_BlendListByInt::GetContextMenuActions(const FGraphNodeContextMenuBuilder& Context) const
+void UAnimGraphNode_BlendListByInt::GetNodeContextMenuActions(UToolMenu* Menu, UGraphNodeContextMenuContext* Context) const
 {
-	if (!Context.bIsDebugging)
+	if (!Context->bIsDebugging)
 	{
-		Context.MenuBuilder->BeginSection("AnimGraphBlendList", NSLOCTEXT("A3Nodes", "BlendListHeader", "BlendList"));
 		{
-			if (Context.Pin != NULL)
+			FToolMenuSection& Section = Menu->AddSection("AnimGraphBlendList", NSLOCTEXT("A3Nodes", "BlendListHeader", "BlendList"));
+			if (Context->Pin != NULL)
 			{
 				// we only do this for normal BlendList/BlendList by enum, BlendList by Bool doesn't support add/remove pins
-				if ( Context.Pin->Direction == EGPD_Input )
+				if ( Context->Pin->Direction == EGPD_Input )
 				{
-					Context.MenuBuilder->AddMenuEntry(FGraphEditorCommands::Get().RemoveBlendListPin);
+					Section.AddMenuEntry(FGraphEditorCommands::Get().RemoveBlendListPin);
 				}
 			}
 			else
 			{
-				Context.MenuBuilder->AddMenuEntry(FGraphEditorCommands::Get().AddBlendListPin);
+				Section.AddMenuEntry(FGraphEditorCommands::Get().AddBlendListPin);
 			}
 		}
-		Context.MenuBuilder->EndSection();	
 	}
 }
 
diff --git a/Engine/Source/Editor/AnimGraph/Private/AnimGraphNode_BlendSpaceEvaluator.cpp b/Engine/Source/Editor/AnimGraph/Private/AnimGraphNode_BlendSpaceEvaluator.cpp
--- a/Engine/Source/Editor/AnimGraph/Private/AnimGraphNode_BlendSpaceEvaluator.cpp
+++ b/Engine/Source/Editor/AnimGraph/Private/AnimGraphNode_BlendSpaceEvaluator.cpp
@@ -1,7 +1,7 @@
 // Copyright 1998-2019 Epic Games, Inc. All Rights Reserved.
 
 #include "AnimGraphNode_BlendSpaceEvaluator.h"
-#include "Framework/MultiBox/MultiBoxBuilder.h"
+#include "ToolMenus.h"
 #include "GraphEditorActions.h"
 #include "Kismet2/CompilerResultsLog.h"
 
@@ -130,17 +130,16 @@ void UAnimGraphNode_BlendSpaceEvaluator::BakeDataDuringCompilation(class FCompil
 	Node.GroupRole = SyncGroup.GroupRole;
 }
 
-void UAnimGraphNode_BlendSpaceEvaluator::GetContextMenuActions(const FGraphNodeContextMenuBuilder& Context) const
+void UAnimGraphNode_BlendSpaceEvaluator::GetNodeContextMenuActions(UToolMenu* Menu, UGraphNodeContextMenuContext* Context) const
 {
-	if (!Context.bIsDebugging)
+	if (!Context->bIsDebugging)
 	{
 		// add an option to convert to single frame
-		Context.MenuBuilder->BeginSection("AnimGraphNodeBlendSpacePlayer", NSLOCTEXT("A3Nodes", "BlendSpaceHeading", "Blend Space"));
 		{
-			Context.MenuBuilder->AddMenuEntry(FGraphEditorCommands::Get().OpenRelatedAsset);
-			Context.MenuBuilder->AddMenuEntry(FGraphEditorCommands::Get().ConvertToBSPlayer);
+			FToolMenuSection& Section = Menu->AddSection("AnimGraphNodeBlendSpacePlayer", NSLOCTEXT("A3Nodes", "BlendSpaceHeading", "Blend Space"));
+			Section.AddMenuEntry(FGraphEditorCommands::Get().OpenRelatedAsset);
+			Section.AddMenuEntry(FGraphEditorCommands::Get().ConvertToBSPlayer);
 		}
-		Context.MenuBuilder->EndSection();
 	}
 }
 
diff --git a/Engine/Source/Editor/AnimGraph/Private/AnimGraphNode_BlendSpacePlayer.cpp b/Engine/Source/Editor/AnimGraph/Private/AnimGraphNode_BlendSpacePlayer.cpp
--- a/Engine/Source/Editor/AnimGraph/Private/AnimGraphNode_BlendSpacePlayer.cpp
+++ b/Engine/Source/Editor/AnimGraph/Private/AnimGraphNode_BlendSpacePlayer.cpp
@@ -3,7 +3,7 @@
 #include "AnimGraphNode_BlendSpacePlayer.h"
 #include "UObject/UObjectHash.h"
 #include "UObject/UObjectIterator.h"
-#include "Framework/MultiBox/MultiBoxBuilder.h"
+#include "ToolMenus.h"
 #include "GraphEditorActions.h"
 #include "Kismet2/CompilerResultsLog.h"
 #include "BlueprintNodeSpawner.h"
@@ -129,17 +129,16 @@ void UAnimGraphNode_BlendSpacePlayer::BakeDataDuringCompilation(class FCompilerR
 	Node.GroupRole = SyncGroup.GroupRole;
 }
 
-void UAnimGraphNode_BlendSpacePlayer::GetContextMenuActions(const FGraphNodeContextMenuBuilder& Context) const
+void UAnimGraphNode_BlendSpacePlayer::GetNodeContextMenuActions(UToolMenu* Menu, UGraphNodeContextMenuContext* Context) const
 {
-	if (!Context.bIsDebugging)
+	if (!Context->bIsDebugging)
 	{
 		// add an option to convert to single frame
-		Context.MenuBuilder->BeginSection("AnimGraphNodeBlendSpaceEvaluator", NSLOCTEXT("A3Nodes", "BlendSpaceHeading", "Blend Space"));
 		{
-			Context.MenuBuilder->AddMenuEntry(FGraphEditorCommands::Get().OpenRelatedAsset);
-			Context.MenuBuilder->AddMenuEntry(FGraphEditorCommands::Get().ConvertToBSEvaluator);
+			FToolMenuSection& Section = Menu->AddSection("AnimGraphNodeBlendSpaceEvaluator", NSLOCTEXT("A3Nodes", "BlendSpaceHeading", "Blend Space"));
+			Section.AddMenuEntry(FGraphEditorCommands::Get().OpenRelatedAsset);
+			Section.AddMenuEntry(FGraphEditorCommands::Get().ConvertToBSEvaluator);
 		}
-		Context.MenuBuilder->EndSection();
 	}
 }
 
diff --git a/Engine/Source/Editor/AnimGraph/Private/AnimGraphNode_LayeredBoneBlend.cpp b/Engine/Source/Editor/AnimGraph/Private/AnimGraphNode_LayeredBoneBlend.cpp
--- a/Engine/Source/Editor/AnimGraph/Private/AnimGraphNode_LayeredBoneBlend.cpp
+++ b/Engine/Source/Editor/AnimGraph/Private/AnimGraphNode_LayeredBoneBlend.cpp
@@ -1,7 +1,7 @@
 // Copyright 1998-2019 Epic Games, Inc. All Rights Reserved.
 
 #include "AnimGraphNode_LayeredBoneBlend.h"
-#include "Framework/MultiBox/MultiBoxBuilder.h"
+#include "ToolMenus.h"
 #include "Kismet2/BlueprintEditorUtils.h"
 
 #include "GraphEditorActions.h"
@@ -68,27 +68,26 @@ void UAnimGraphNode_LayeredBoneBlend::RemovePinFromBlendByFilter(UEdGraphPin* Pi
 	}
 }
 
-void UAnimGraphNode_LayeredBoneBlend::GetContextMenuActions(const FGraphNodeContextMenuBuilder& Context) const
+void UAnimGraphNode_LayeredBoneBlend::GetNodeContextMenuActions(UToolMenu* Menu, UGraphNodeContextMenuContext* Context) const
 {
-	if (!Context.bIsDebugging)
+	if (!Context->bIsDebugging)
 	{
-		Context.MenuBuilder->BeginSection("AnimGraphNodeLayeredBoneblend", LOCTEXT("LayeredBoneBlend", "Layered Bone Blend"));
 		{
-			if (Context.Pin != NULL)
+			FToolMenuSection& Section = Menu->AddSection("AnimGraphNodeLayeredBoneblend", LOCTEXT("LayeredBoneBlend", "Layered Bone Blend"));
+			if (Context->Pin != NULL)
 			{
 				// we only do this for normal BlendList/BlendList by enum, BlendList by Bool doesn't support add/remove pins
-				if (Context.Pin->Direction == EGPD_Input)
+				if (Context->Pin->Direction == EGPD_Input)
 				{
 					//@TODO: Only offer this option on arrayed pins
-					Context.MenuBuilder->AddMenuEntry(FGraphEditorCommands::Get().RemoveBlendListPin);
+					Section.AddMenuEntry(FGraphEditorCommands::Get().RemoveBlendListPin);
 				}
 			}
 			else
 			{
-				Context.MenuBuilder->AddMenuEntry(FGraphEditorCommands::Get().AddBlendListPin);
+				Section.AddMenuEntry(FGraphEditorCommands::Get().AddBlendListPin);
 			}
 		}
-		Context.MenuBuilder->EndSection();
 	}
 }
 
diff --git a/Engine/Source/Editor/AnimGraph/Private/AnimGraphNode_ModifyCurve.cpp b/Engine/Source/Editor/AnimGraph/Private/AnimGraphNode_ModifyCurve.cpp
--- a/Engine/Source/Editor/AnimGraph/Private/AnimGraphNode_ModifyCurve.cpp
+++ b/Engine/Source/Editor/AnimGraph/Private/AnimGraphNode_ModifyCurve.cpp
@@ -8,7 +8,7 @@
 #include "AnimationGraphSchema.h"
 #include "BlueprintActionDatabaseRegistrar.h"
 #include "Framework/Commands/UIAction.h"
-#include "Framework/MultiBox/MultiBoxBuilder.h"
+#include "ToolMenus.h"
 #include "Kismet2/BlueprintEditorUtils.h"
 
 #define LOCTEXT_NAMESPACE "ModifyCurve"
@@ -74,34 +74,36 @@ void UAnimGraphNode_ModifyCurve::GetRemoveCurveMenuActions(FMenuBuilder& MenuBui
 }
 
 
-void UAnimGraphNode_ModifyCurve::GetContextMenuActions(const FGraphNodeContextMenuBuilder& Context) const
+void UAnimGraphNode_ModifyCurve::GetNodeContextMenuActions(UToolMenu* Menu, UGraphNodeContextMenuContext* Context) const
 {
-	if (!Context.bIsDebugging)
+	if (!Context->bIsDebugging)
 	{
-		Context.MenuBuilder->BeginSection("AnimGraphNodeModifyCurve", LOCTEXT("ModifyCurve", "Modify Curve"));
+		FToolMenuSection& Section = Menu->AddSection("AnimGraphNodeModifyCurve", LOCTEXT("ModifyCurve", "Modify Curve"));
 
 		// Clicked pin
-		if (Context.Pin != NULL)
+		if (Context->Pin != NULL)
 		{
 			// Get proeprty from pin
 			UProperty* AssociatedProperty;
 			int32 ArrayIndex;
-			GetPinAssociatedProperty(GetFNodeType(), Context.Pin, /*out*/ AssociatedProperty, /*out*/ ArrayIndex);
+			GetPinAssociatedProperty(GetFNodeType(), Context->Pin, /*out*/ AssociatedProperty, /*out*/ ArrayIndex);
 			FName PinPropertyName = AssociatedProperty->GetFName();
 
-			if (PinPropertyName  == GET_MEMBER_NAME_CHECKED(FAnimNode_ModifyCurve, CurveValues) && Context.Pin->Direction == EGPD_Input)
+			if (PinPropertyName  == GET_MEMBER_NAME_CHECKED(FAnimNode_ModifyCurve, CurveValues) && Context->Pin->Direction == EGPD_Input)
 			{
-				FString PinName = Context.Pin->PinFriendlyName.ToString();
+				FString PinName = Context->Pin->PinFriendlyName.ToString();
 				FUIAction Action = FUIAction( FExecuteAction::CreateUObject(const_cast<UAnimGraphNode_ModifyCurve*>(this), &UAnimGraphNode_ModifyCurve::RemoveCurvePin, FName(*PinName)) );
 				FText RemovePinLabelText = FText::Format(LOCTEXT("RemoveThisPin", "Remove This Curve Pin: {0}"), FText::FromString(PinName));
-				Context.MenuBuilder->AddMenuEntry(RemovePinLabelText, LOCTEXT("RemoveThisPinTooltip", "Remove this curve pin from this node"), FSlateIcon(), Action);
+				Section.AddMenuEntry("RemoveThisPin", RemovePinLabelText, LOCTEXT("RemoveThisPinTooltip", "Remove this curve pin from this node"), FSlateIcon(), Action);
 			}
 		}
 
 		// If we have more curves to add, create submenu to offer them
 		if (GetCurvesToAdd().Num() > 0)
 		{
-			Context.MenuBuilder->AddSubMenu(
+			Section.AddSubMenu(
+				Menu->GetMenuName(),
+				"AddCurvePin",
 				LOCTEXT("AddCurvePin", "Add Curve Pin"),
 				LOCTEXT("AddCurvePinTooltip", "Add a new pin to drive a curve"),
 				FNewMenuDelegate::CreateUObject(this, &UAnimGraphNode_ModifyCurve::GetAddCurveMenuActions));
@@ -110,14 +112,14 @@ void UAnimGraphNode_ModifyCurve::GetContextMenuActions(const FGraphNodeContextMe
 		// If we have curves to remove, create submenu to offer them
 		if (Node.CurveNames.Num() > 0)
 		{
-			Context.MenuBuilder->AddSubMenu(
+			Section.AddSubMenu(
+				Menu->GetMenuName(),
+				"RemoveCurvePin",
 				LOCTEXT("RemoveCurvePin", "Remove Curve Pin"),
 				LOCTEXT("RemoveCurvePinTooltip", "Remove a pin driving a curve"),
 				FNewMenuDelegate::CreateUObject(this, &UAnimGraphNode_ModifyCurve::GetRemoveCurveMenuActions));
 		}
 	}
-
-	Context.MenuBuilder->EndSection();
 }
 
 void UAnimGraphNode_ModifyCurve::RemoveCurvePin(FName CurveName)
diff --git a/Engine/Source/Editor/AnimGraph/Private/AnimGraphNode_MultiWayBlend.cpp b/Engine/Source/Editor/AnimGraph/Private/AnimGraphNode_MultiWayBlend.cpp
--- a/Engine/Source/Editor/AnimGraph/Private/AnimGraphNode_MultiWayBlend.cpp
+++ b/Engine/Source/Editor/AnimGraph/Private/AnimGraphNode_MultiWayBlend.cpp
@@ -4,7 +4,7 @@
 #include "Kismet2/BlueprintEditorUtils.h"
 #include "GraphEditorActions.h"
 #include "ScopedTransaction.h"
-#include "Framework/MultiBox/MultiBoxBuilder.h"
+#include "ToolMenus.h"
 
 
 /////////////////////////////////////////////////////
@@ -37,26 +37,25 @@ FText UAnimGraphNode_MultiWayBlend::GetNodeTitle(ENodeTitleType::Type TitleType)
 	return LOCTEXT("Blend", "Blend Multi");
 }
 
-void UAnimGraphNode_MultiWayBlend::GetContextMenuActions(const FGraphNodeContextMenuBuilder& Context) const
+void UAnimGraphNode_MultiWayBlend::GetNodeContextMenuActions(UToolMenu* Menu, UGraphNodeContextMenuContext* Context) const
 {
-	if (!Context.bIsDebugging)
+	if (!Context->bIsDebugging)
 	{
-		Context.MenuBuilder->BeginSection("AnimGraphBlendMulti", LOCTEXT("BlendMultiHeader", "BlendMulti"));
 		{
-			if (Context.Pin != NULL)
+			FToolMenuSection& Section = Menu->AddSection("AnimGraphBlendMulti", LOCTEXT("BlendMultiHeader", "BlendMulti"));
+			if (Context->Pin != NULL)
 			{
 				// we only do this for normal BlendMulti/BlendMulti by enum, BlendMulti by Bool doesn't support add/remove pins
-				if (Context.Pin->Direction == EGPD_Input)
+				if (Context->Pin->Direction == EGPD_Input)
 				{
-					Context.MenuBuilder->AddMenuEntry(FGraphEditorCommands::Get().RemoveBlendListPin);
+					Section.AddMenuEntry(FGraphEditorCommands::Get().RemoveBlendListPin);
 				}
 			}
 			else
 			{
-				Context.MenuBuilder->AddMenuEntry(FGraphEditorCommands::Get().AddBlendListPin);
+				Section.AddMenuEntry(FGraphEditorCommands::Get().AddBlendListPin);
 			}
 		}
-		Context.MenuBuilder->EndSection();
 	}
 }
 
diff --git a/Engine/Source/Editor/AnimGraph/Private/AnimGraphNode_PoseBlendNode.cpp b/Engine/Source/Editor/AnimGraph/Private/AnimGraphNode_PoseBlendNode.cpp
--- a/Engine/Source/Editor/AnimGraph/Private/AnimGraphNode_PoseBlendNode.cpp
+++ b/Engine/Source/Editor/AnimGraph/Private/AnimGraphNode_PoseBlendNode.cpp
@@ -3,7 +3,7 @@
 #include "AnimGraphNode_PoseBlendNode.h"
 #include "EdGraphSchema_K2_Actions.h"
 #include "Modules/ModuleManager.h"
-#include "Framework/MultiBox/MultiBoxBuilder.h"
+#include "ToolMenus.h"
 
 #include "GraphEditorActions.h"
 #include "ARFilter.h"
@@ -244,16 +244,15 @@ bool UAnimGraphNode_PoseBlendNode::DoesSupportTimeForTransitionGetter() const
 	return false;
 }
 
-void UAnimGraphNode_PoseBlendNode::GetContextMenuActions(const FGraphNodeContextMenuBuilder& Context) const
+void UAnimGraphNode_PoseBlendNode::GetNodeContextMenuActions(UToolMenu* Menu, UGraphNodeContextMenuContext* Context) const
 {
-	if (!Context.bIsDebugging)
+	if (!Context->bIsDebugging)
 	{
 		// add an option to convert to single frame
-		Context.MenuBuilder->BeginSection("AnimGraphNodePoseBlender", LOCTEXT("PoseBlenderHeading", "Pose Blender"));
 		{
-			Context.MenuBuilder->AddMenuEntry(FGraphEditorCommands::Get().ConvertToPoseByName);
+			FToolMenuSection& Section = Menu->AddSection("AnimGraphNodePoseBlender", LOCTEXT("PoseBlenderHeading", "Pose Blender"));
+			Section.AddMenuEntry(FGraphEditorCommands::Get().ConvertToPoseByName);
 		}
-		Context.MenuBuilder->EndSection();
 	}
 }
 
diff --git a/Engine/Source/Editor/AnimGraph/Private/AnimGraphNode_PoseByName.cpp b/Engine/Source/Editor/AnimGraph/Private/AnimGraphNode_PoseByName.cpp
--- a/Engine/Source/Editor/AnimGraph/Private/AnimGraphNode_PoseByName.cpp
+++ b/Engine/Source/Editor/AnimGraph/Private/AnimGraphNode_PoseByName.cpp
@@ -1,7 +1,7 @@
 // Copyright 1998-2019 Epic Games, Inc. All Rights Reserved.
 
 #include "AnimGraphNode_PoseByName.h"
-#include "Framework/MultiBox/MultiBoxBuilder.h"
+#include "ToolMenus.h"
 
 #include "Kismet2/CompilerResultsLog.h"
 #include "GraphEditorActions.h"
@@ -148,16 +148,15 @@ UAnimationAsset* UAnimGraphNode_PoseByName::GetAnimationAsset() const
 	return PoseAsset;
 }
 
-void UAnimGraphNode_PoseByName::GetContextMenuActions(const FGraphNodeContextMenuBuilder& Context) const
+void UAnimGraphNode_PoseByName::GetNodeContextMenuActions(UToolMenu* Menu, UGraphNodeContextMenuContext* Context) const
 {
-	if (!Context.bIsDebugging)
+	if (!Context->bIsDebugging)
 	{
 		// add an option to convert to single frame
-		Context.MenuBuilder->BeginSection("AnimGraphNodePoseByName", LOCTEXT("PoseByNameHeading", "Pose By Name"));
 		{
-			Context.MenuBuilder->AddMenuEntry(FGraphEditorCommands::Get().ConvertToPoseBlender);
+			FToolMenuSection& Section = Menu->AddSection("AnimGraphNodePoseByName", LOCTEXT("PoseByNameHeading", "Pose By Name"));
+			Section.AddMenuEntry(FGraphEditorCommands::Get().ConvertToPoseBlender);
 		}
-		Context.MenuBuilder->EndSection();
 	}
 }
 
diff --git a/Engine/Source/Editor/AnimGraph/Private/AnimGraphNode_RotationOffsetBlendSpace.cpp b/Engine/Source/Editor/AnimGraph/Private/AnimGraphNode_RotationOffsetBlendSpace.cpp
--- a/Engine/Source/Editor/AnimGraph/Private/AnimGraphNode_RotationOffsetBlendSpace.cpp
+++ b/Engine/Source/Editor/AnimGraph/Private/AnimGraphNode_RotationOffsetBlendSpace.cpp
@@ -3,7 +3,7 @@
 #include "AnimGraphNode_RotationOffsetBlendSpace.h"
 #include "UObject/UObjectHash.h"
 #include "UObject/UObjectIterator.h"
-#include "Framework/MultiBox/MultiBoxBuilder.h"
+#include "ToolMenus.h"
 #include "GraphEditorActions.h"
 #include "Kismet2/CompilerResultsLog.h"
 #include "BlueprintNodeSpawner.h"
@@ -185,17 +185,16 @@ void UAnimGraphNode_RotationOffsetBlendSpace::ValidateAnimNodeDuringCompilation(
 	}
 }
 
-void UAnimGraphNode_RotationOffsetBlendSpace::GetContextMenuActions(const FGraphNodeContextMenuBuilder& Context) const
+void UAnimGraphNode_RotationOffsetBlendSpace::GetNodeContextMenuActions(UToolMenu* Menu, UGraphNodeContextMenuContext* Context) const
 {
-	if (!Context.bIsDebugging)
+	if (!Context->bIsDebugging)
 	{
 		// add an option to convert to single frame
-		Context.MenuBuilder->BeginSection("AnimGraphNodeBlendSpacePlayer", NSLOCTEXT("A3Nodes", "BlendSpaceHeading", "Blend Space"));
 		{
-			Context.MenuBuilder->AddMenuEntry(FGraphEditorCommands::Get().OpenRelatedAsset);
-			Context.MenuBuilder->AddMenuEntry(FGraphEditorCommands::Get().ConvertToAimOffsetLookAt);
+			FToolMenuSection& Section = Menu->AddSection("AnimGraphNodeBlendSpacePlayer", NSLOCTEXT("A3Nodes", "BlendSpaceHeading", "Blend Space"));
+			Section.AddMenuEntry(FGraphEditorCommands::Get().OpenRelatedAsset);
+			Section.AddMenuEntry(FGraphEditorCommands::Get().ConvertToAimOffsetLookAt);
 		}
-		Context.MenuBuilder->EndSection();
 	}
 }
 
diff --git a/Engine/Source/Editor/AnimGraph/Private/AnimGraphNode_SequenceEvaluator.cpp b/Engine/Source/Editor/AnimGraph/Private/AnimGraphNode_SequenceEvaluator.cpp
--- a/Engine/Source/Editor/AnimGraph/Private/AnimGraphNode_SequenceEvaluator.cpp
+++ b/Engine/Source/Editor/AnimGraph/Private/AnimGraphNode_SequenceEvaluator.cpp
@@ -1,7 +1,7 @@
 // Copyright 1998-2019 Epic Games, Inc. All Rights Reserved.
 
 #include "AnimGraphNode_SequenceEvaluator.h"
-#include "Framework/MultiBox/MultiBoxBuilder.h"
+#include "ToolMenus.h"
 
 #include "Kismet2/CompilerResultsLog.h"
 #include "GraphEditorActions.h"
@@ -144,17 +144,16 @@ void UAnimGraphNode_SequenceEvaluator::ValidateAnimNodeDuringCompilation(class U
 	}
 }
 
-void UAnimGraphNode_SequenceEvaluator::GetContextMenuActions(const FGraphNodeContextMenuBuilder& Context) const
+void UAnimGraphNode_SequenceEvaluator::GetNodeContextMenuActions(UToolMenu* Menu, UGraphNodeContextMenuContext* Context) const
 {
-	if (!Context.bIsDebugging)
+	if (!Context->bIsDebugging)
 	{
 		// add an option to convert to a regular sequence player
-		Context.MenuBuilder->BeginSection("AnimGraphNodeSequenceEvaluator", NSLOCTEXT("A3Nodes", "SequenceEvaluatorHeading", "Sequence Evaluator"));
 		{
-			Context.MenuBuilder->AddMenuEntry(FGraphEditorCommands::Get().OpenRelatedAsset);
-			Context.MenuBuilder->AddMenuEntry(FGraphEditorCommands::Get().ConvertToSeqPlayer);
+			FToolMenuSection& Section = Menu->AddSection("AnimGraphNodeSequenceEvaluator", NSLOCTEXT("A3Nodes", "SequenceEvaluatorHeading", "Sequence Evaluator"));
+			Section.AddMenuEntry(FGraphEditorCommands::Get().OpenRelatedAsset);
+			Section.AddMenuEntry(FGraphEditorCommands::Get().ConvertToSeqPlayer);
 		}
-		Context.MenuBuilder->EndSection();
 	}
 }
 
diff --git a/Engine/Source/Editor/AnimGraph/Private/AnimGraphNode_SequencePlayer.cpp b/Engine/Source/Editor/AnimGraph/Private/AnimGraphNode_SequencePlayer.cpp
--- a/Engine/Source/Editor/AnimGraph/Private/AnimGraphNode_SequencePlayer.cpp
+++ b/Engine/Source/Editor/AnimGraph/Private/AnimGraphNode_SequencePlayer.cpp
@@ -3,7 +3,7 @@
 #include "AnimGraphNode_SequencePlayer.h"
 #include "EdGraphSchema_K2_Actions.h"
 #include "Modules/ModuleManager.h"
-#include "Framework/MultiBox/MultiBoxBuilder.h"
+#include "ToolMenus.h"
 
 #include "Kismet2/CompilerResultsLog.h"
 #include "GraphEditorActions.h"
@@ -344,17 +344,16 @@ void UAnimGraphNode_SequencePlayer::ValidateAnimNodeDuringCompilation(class USke
 	}
 }
 
-void UAnimGraphNode_SequencePlayer::GetContextMenuActions(const FGraphNodeContextMenuBuilder& Context) const
+void UAnimGraphNode_SequencePlayer::GetNodeContextMenuActions(UToolMenu* Menu, UGraphNodeContextMenuContext* Context) const
 {
-	if (!Context.bIsDebugging)
+	if (!Context->bIsDebugging)
 	{
 		// add an option to convert to single frame
-		Context.MenuBuilder->BeginSection("AnimGraphNodeSequencePlayer", NSLOCTEXT("A3Nodes", "SequencePlayerHeading", "Sequence Player"));
 		{
-			Context.MenuBuilder->AddMenuEntry(FGraphEditorCommands::Get().OpenRelatedAsset);
-			Context.MenuBuilder->AddMenuEntry(FGraphEditorCommands::Get().ConvertToSeqEvaluator);
+			FToolMenuSection& Section = Menu->AddSection("AnimGraphNodeSequencePlayer", NSLOCTEXT("A3Nodes", "SequencePlayerHeading", "Sequence Player"));
+			Section.AddMenuEntry(FGraphEditorCommands::Get().OpenRelatedAsset);
+			Section.AddMenuEntry(FGraphEditorCommands::Get().ConvertToSeqEvaluator);
 		}
-		Context.MenuBuilder->EndSection();
 	}
 }
 
diff --git a/Engine/Source/Editor/AudioEditor/Classes/SoundCueGraph/SoundCueGraphNode.h b/Engine/Source/Editor/AudioEditor/Classes/SoundCueGraph/SoundCueGraphNode.h
--- a/Engine/Source/Editor/AudioEditor/Classes/SoundCueGraph/SoundCueGraphNode.h
+++ b/Engine/Source/Editor/AudioEditor/Classes/SoundCueGraph/SoundCueGraphNode.h
@@ -42,7 +42,7 @@ class USoundCueGraphNode : public USoundCueGraphNode_Base
 	// UEdGraphNode interface
 	virtual FText GetNodeTitle(ENodeTitleType::Type TitleType) const override;
 	virtual void PrepareForCopying() override;
-	virtual void GetContextMenuActions(const FGraphNodeContextMenuBuilder& Context) const override;
+	virtual void GetNodeContextMenuActions(class UToolMenu* Menu, class UGraphNodeContextMenuContext* Context) const override;
 	virtual FText GetTooltipText() const override;
 	virtual FString GetDocumentationExcerptName() const override;
 	// End of UEdGraphNode interface
diff --git a/Engine/Source/Editor/AudioEditor/Classes/SoundCueGraph/SoundCueGraphNode_Root.h b/Engine/Source/Editor/AudioEditor/Classes/SoundCueGraph/SoundCueGraphNode_Root.h
--- a/Engine/Source/Editor/AudioEditor/Classes/SoundCueGraph/SoundCueGraphNode_Root.h
+++ b/Engine/Source/Editor/AudioEditor/Classes/SoundCueGraph/SoundCueGraphNode_Root.h
@@ -18,7 +18,7 @@ class USoundCueGraphNode_Root : public USoundCueGraphNode_Base
 	virtual FText GetNodeTitle(ENodeTitleType::Type TitleType) const override;
 	virtual bool CanUserDeleteNode() const override { return false; }
 	virtual bool CanDuplicateNode() const override { return false; }
-	virtual void GetContextMenuActions(const FGraphNodeContextMenuBuilder& Context) const override;
+	virtual void GetNodeContextMenuActions(class UToolMenu* Menu, class UGraphNodeContextMenuContext* Context) const override;
 	// End of UEdGraphNode interface
 
 	// USoundCueGraphNode_Base interface
diff --git a/Engine/Source/Editor/AudioEditor/Private/SoundCueGraphNode.cpp b/Engine/Source/Editor/AudioEditor/Private/SoundCueGraphNode.cpp
--- a/Engine/Source/Editor/AudioEditor/Private/SoundCueGraphNode.cpp
+++ b/Engine/Source/Editor/AudioEditor/Private/SoundCueGraphNode.cpp
@@ -1,7 +1,7 @@
 // Copyright 1998-2019 Epic Games, Inc. All Rights Reserved.
 
 #include "SoundCueGraph/SoundCueGraphNode.h"
-#include "Framework/MultiBox/MultiBoxBuilder.h"
+#include "ToolMenus.h"
 #include "Editor/EditorEngine.h"
 #include "SoundCueGraph/SoundCueGraph.h"
 #include "Sound/SoundNodeWavePlayer.h"
@@ -211,65 +211,68 @@ void USoundCueGraphNode::CreateInputPins()
 	}
 }
 
-void USoundCueGraphNode::GetContextMenuActions(const FGraphNodeContextMenuBuilder& Context) const
+void USoundCueGraphNode::GetNodeContextMenuActions(UToolMenu* Menu, UGraphNodeContextMenuContext* Context) const
 {
-	if (Context.Pin)
+	if (Context->Pin)
 	{
 		// If on an input that can be deleted, show option
-		if(Context.Pin->Direction == EGPD_Input && SoundNode->ChildNodes.Num() > SoundNode->GetMinChildNodes())
+		if(Context->Pin->Direction == EGPD_Input && SoundNode->ChildNodes.Num() > SoundNode->GetMinChildNodes())
 		{
-			Context.MenuBuilder->AddMenuEntry(FSoundCueGraphEditorCommands::Get().DeleteInput);
+			FToolMenuSection& Section = Menu->AddSection("SoundCueGraphDeleteInput");
+			Section.AddMenuEntry(FSoundCueGraphEditorCommands::Get().DeleteInput);
 		}
 	}
-	else if (Context.Node)
+	else if (Context->Node)
 	{
-		Context.MenuBuilder->AddSubMenu(LOCTEXT("AlignmentHeader", "Alignment"), FText(), FNewMenuDelegate::CreateLambda([](FMenuBuilder& InMenuBuilder) {
-
-			InMenuBuilder.BeginSection("EdGraphSchemaAlignment", LOCTEXT("AlignHeader", "Align"));
-			InMenuBuilder.AddMenuEntry(FGraphEditorCommands::Get().AlignNodesTop);
-			InMenuBuilder.AddMenuEntry(FGraphEditorCommands::Get().AlignNodesMiddle);
-			InMenuBuilder.AddMenuEntry(FGraphEditorCommands::Get().AlignNodesBottom);
-			InMenuBuilder.AddMenuEntry(FGraphEditorCommands::Get().AlignNodesLeft);
-			InMenuBuilder.AddMenuEntry(FGraphEditorCommands::Get().AlignNodesCenter);
-			InMenuBuilder.AddMenuEntry(FGraphEditorCommands::Get().AlignNodesRight);
-			InMenuBuilder.AddMenuEntry(FGraphEditorCommands::Get().StraightenConnections);
-			InMenuBuilder.EndSection();
-
-			InMenuBuilder.BeginSection("EdGraphSchemaDistribution", LOCTEXT("DistributionHeader", "Distribution"));
-			InMenuBuilder.AddMenuEntry(FGraphEditorCommands::Get().DistributeNodesHorizontally);
-			InMenuBuilder.AddMenuEntry(FGraphEditorCommands::Get().DistributeNodesVertically);
-			InMenuBuilder.EndSection();
-		}));
-
-		Context.MenuBuilder->BeginSection("SoundCueGraphNodeEdit");
 		{
-			Context.MenuBuilder->AddMenuEntry( FGenericCommands::Get().Delete );
-			Context.MenuBuilder->AddMenuEntry( FGenericCommands::Get().Cut );
-			Context.MenuBuilder->AddMenuEntry( FGenericCommands::Get().Copy );
-			Context.MenuBuilder->AddMenuEntry( FGenericCommands::Get().Duplicate );
+			FToolMenuSection& Section = Menu->AddSection("SoundCueGraphNodeAlignment");
+			Section.AddSubMenu("Alignment", LOCTEXT("AlignmentHeader", "Alignment"), FText(), FNewToolMenuDelegate::CreateLambda([](UToolMenu* SubMenu)
+			{
+				{
+					FToolMenuSection& SubMenuSection = SubMenu->AddSection("EdGraphSchemaAlignment", LOCTEXT("AlignHeader", "Align"));
+					SubMenuSection.AddMenuEntry(FGraphEditorCommands::Get().AlignNodesTop);
+					SubMenuSection.AddMenuEntry(FGraphEditorCommands::Get().AlignNodesMiddle);
+					SubMenuSection.AddMenuEntry(FGraphEditorCommands::Get().AlignNodesBottom);
+					SubMenuSection.AddMenuEntry(FGraphEditorCommands::Get().AlignNodesLeft);
+					SubMenuSection.AddMenuEntry(FGraphEditorCommands::Get().AlignNodesCenter);
+					SubMenuSection.AddMenuEntry(FGraphEditorCommands::Get().AlignNodesRight);
+					SubMenuSection.AddMenuEntry(FGraphEditorCommands::Get().StraightenConnections);
+				}
+
+				{
+					FToolMenuSection& SubMenuSection = SubMenu->AddSection("EdGraphSchemaDistribution", LOCTEXT("DistributionHeader", "Distribution"));
+					SubMenuSection.AddMenuEntry(FGraphEditorCommands::Get().DistributeNodesHorizontally);
+					SubMenuSection.AddMenuEntry(FGraphEditorCommands::Get().DistributeNodesVertically);
+				}
+			}));
 		}
-		Context.MenuBuilder->EndSection();
 
-		Context.MenuBuilder->BeginSection("SoundCueGraphNodeAddPlaySync");
 		{
+			FToolMenuSection& Section = Menu->AddSection("SoundCueGraphNodeEdit");
+			Section.AddMenuEntry( FGenericCommands::Get().Delete );
+			Section.AddMenuEntry( FGenericCommands::Get().Cut );
+			Section.AddMenuEntry( FGenericCommands::Get().Copy );
+			Section.AddMenuEntry( FGenericCommands::Get().Duplicate );
+		}
+
+		{
+			FToolMenuSection& Section = Menu->AddSection("SoundCueGraphNodeAddPlaySync");
 			if (CanAddInputPin())
 			{
-				Context.MenuBuilder->AddMenuEntry(FSoundCueGraphEditorCommands::Get().AddInput);
+				Section.AddMenuEntry(FSoundCueGraphEditorCommands::Get().AddInput);
 			}
 
-			Context.MenuBuilder->AddMenuEntry(FSoundCueGraphEditorCommands::Get().PlayNode);
+			Section.AddMenuEntry(FSoundCueGraphEditorCommands::Get().PlayNode);
 
 			if ( Cast<USoundNodeWavePlayer>(SoundNode) )
 			{
-				Context.MenuBuilder->AddMenuEntry(FSoundCueGraphEditorCommands::Get().BrowserSync);
+				Section.AddMenuEntry(FSoundCueGraphEditorCommands::Get().BrowserSync);
 			}
 			else if (Cast<USoundNodeDialoguePlayer>(SoundNode))
 			{
-				Context.MenuBuilder->AddMenuEntry(FSoundCueGraphEditorCommands::Get().BrowserSync);
+				Section.AddMenuEntry(FSoundCueGraphEditorCommands::Get().BrowserSync);
 			}
 		}
-
-		Context.MenuBuilder->EndSection();
 	}
 }
 
diff --git a/Engine/Source/Editor/AudioEditor/Private/SoundCueGraphNode_Root.cpp b/Engine/Source/Editor/AudioEditor/Private/SoundCueGraphNode_Root.cpp
--- a/Engine/Source/Editor/AudioEditor/Private/SoundCueGraphNode_Root.cpp
+++ b/Engine/Source/Editor/AudioEditor/Private/SoundCueGraphNode_Root.cpp
@@ -5,7 +5,7 @@
 =============================================================================*/
 
 #include "SoundCueGraph/SoundCueGraphNode_Root.h"
-#include "Framework/MultiBox/MultiBoxBuilder.h"
+#include "ToolMenus.h"
 #include "GraphEditorSettings.h"
 #include "SoundCueGraphEditorCommands.h"
 
@@ -39,18 +39,17 @@ void USoundCueGraphNode_Root::CreateInputPins()
 	CreatePin(EGPD_Input, TEXT("SoundNode"), TEXT("Root"), NAME_None);
 }
 
-void USoundCueGraphNode_Root::GetContextMenuActions(const FGraphNodeContextMenuBuilder& Context) const
+void USoundCueGraphNode_Root::GetNodeContextMenuActions(UToolMenu* Menu, UGraphNodeContextMenuContext* Context) const
 {
-	if (Context.Pin)
+	if (Context->Pin)
 	{
 	}
-	else if (Context.Node)
+	else if (Context->Node)
 	{
-		Context.MenuBuilder->BeginSection("SoundCueGraphNodePlay");
 		{
-			Context.MenuBuilder->AddMenuEntry(FSoundCueGraphEditorCommands::Get().PlayNode);
+			FToolMenuSection& Section = Menu->AddSection("SoundCueGraphNodePlay");
+			Section.AddMenuEntry(FSoundCueGraphEditorCommands::Get().PlayNode);
 		}
-		Context.MenuBuilder->EndSection();
 	}
 }
 
diff --git a/Engine/Source/Editor/BehaviorTreeEditor/Classes/BehaviorTreeGraphNode.h b/Engine/Source/Editor/BehaviorTreeEditor/Classes/BehaviorTreeGraphNode.h
--- a/Engine/Source/Editor/BehaviorTreeEditor/Classes/BehaviorTreeGraphNode.h
+++ b/Engine/Source/Editor/BehaviorTreeEditor/Classes/BehaviorTreeGraphNode.h
@@ -128,14 +128,14 @@ class UBehaviorTreeGraphNode : public UAIGraphNode
 protected:
 
 	/** creates add decorator... submenu */
-	void CreateAddDecoratorSubMenu(class FMenuBuilder& MenuBuilder, UEdGraph* Graph) const;
+	void CreateAddDecoratorSubMenu(class UToolMenu* Menu, UEdGraph* Graph) const;
 
 	/** creates add service... submenu */
-	void CreateAddServiceSubMenu(class FMenuBuilder& MenuBuilder, UEdGraph* Graph) const;
+	void CreateAddServiceSubMenu(class UToolMenu* Menu, UEdGraph* Graph) const;
 
 	/** add right click menu to create subnodes: Decorators */
-	void AddContextMenuActionsDecorators(const FGraphNodeContextMenuBuilder& Context) const;
+	void AddContextMenuActionsDecorators(class UToolMenu* Menu, const FName SectionName, class UGraphNodeContextMenuContext* Context) const;
 
 	/** add right click menu to create subnodes: Services */
-	void AddContextMenuActionsServices(const FGraphNodeContextMenuBuilder& Context) const;
+	void AddContextMenuActionsServices(class UToolMenu* Menu, const FName SectionName, class UGraphNodeContextMenuContext* Context) const;
 };
diff --git a/Engine/Source/Editor/BehaviorTreeEditor/Classes/BehaviorTreeGraphNode_Composite.h b/Engine/Source/Editor/BehaviorTreeEditor/Classes/BehaviorTreeGraphNode_Composite.h
--- a/Engine/Source/Editor/BehaviorTreeEditor/Classes/BehaviorTreeGraphNode_Composite.h
+++ b/Engine/Source/Editor/BehaviorTreeEditor/Classes/BehaviorTreeGraphNode_Composite.h
@@ -18,7 +18,7 @@ class BEHAVIORTREEEDITOR_API UBehaviorTreeGraphNode_Composite : public UBehavior
 	virtual bool RefreshNodeClass() override{ return false; }
 
 	/** Gets a list of actions that can be done to this particular node */
-	virtual void GetContextMenuActions(const FGraphNodeContextMenuBuilder& Context) const override;
+	virtual void GetNodeContextMenuActions(class UToolMenu* Menu, class UGraphNodeContextMenuContext* Context) const override;
 
 	/** check if node can accept breakpoints */
 	virtual bool CanPlaceBreakpoints() const override { return true; }
diff --git a/Engine/Source/Editor/BehaviorTreeEditor/Classes/BehaviorTreeGraphNode_Task.h b/Engine/Source/Editor/BehaviorTreeEditor/Classes/BehaviorTreeGraphNode_Task.h
--- a/Engine/Source/Editor/BehaviorTreeEditor/Classes/BehaviorTreeGraphNode_Task.h
+++ b/Engine/Source/Editor/BehaviorTreeEditor/Classes/BehaviorTreeGraphNode_Task.h
@@ -15,7 +15,7 @@ class UBehaviorTreeGraphNode_Task : public UBehaviorTreeGraphNode
 	virtual void AllocateDefaultPins() override;
 	virtual FText GetNodeTitle(ENodeTitleType::Type TitleType) const override;
 	/** Gets a list of actions that can be done to this particular node */
-	virtual void GetContextMenuActions(const FGraphNodeContextMenuBuilder& Context) const override;
+	virtual void GetNodeContextMenuActions(class UToolMenu* Menu, class UGraphNodeContextMenuContext* Context) const override;
 
 	virtual bool CanPlaceBreakpoints() const override { return true; }
 };
diff --git a/Engine/Source/Editor/BehaviorTreeEditor/Private/BehaviorTreeGraphNode.cpp b/Engine/Source/Editor/BehaviorTreeEditor/Private/BehaviorTreeGraphNode.cpp
--- a/Engine/Source/Editor/BehaviorTreeEditor/Private/BehaviorTreeGraphNode.cpp
+++ b/Engine/Source/Editor/BehaviorTreeEditor/Private/BehaviorTreeGraphNode.cpp
@@ -2,7 +2,7 @@
 
 #include "BehaviorTreeGraphNode.h"
 #include "Widgets/DeclarativeSyntaxSupport.h"
-#include "Framework/MultiBox/MultiBoxBuilder.h"
+#include "ToolMenus.h"
 #include "BehaviorTree/BTNode.h"
 #include "BehaviorTree/BTDecorator.h"
 #include "BehaviorTree/BTService.h"
@@ -327,44 +327,50 @@ void UBehaviorTreeGraphNode::InsertSubNodeAt(UAIGraphNode* SubNode, int32 DropIn
 	}
 }
 
-void UBehaviorTreeGraphNode::CreateAddDecoratorSubMenu(class FMenuBuilder& MenuBuilder, UEdGraph* Graph) const
+void UBehaviorTreeGraphNode::CreateAddDecoratorSubMenu(UToolMenu* Menu, UEdGraph* Graph) const
 {
-	TSharedRef<SGraphEditorActionMenuAI> Menu =	
+	TSharedRef<SGraphEditorActionMenuAI> Widget =
 		SNew(SGraphEditorActionMenuAI)
 		.GraphObj( Graph )
 		.GraphNode((UBehaviorTreeGraphNode*)this)
 		.SubNodeFlags(ESubNode::Decorator)
 		.AutoExpandActionMenu(true);
 
-	MenuBuilder.AddWidget(Menu,FText(),true);
+	FToolMenuSection& Section = Menu->FindOrAddSection("Section");
+	Section.AddEntry(FToolMenuEntry::InitWidget("DecoratorWidget", Widget, FText(), true));
 }
 
-void UBehaviorTreeGraphNode::CreateAddServiceSubMenu(class FMenuBuilder& MenuBuilder, UEdGraph* Graph) const
+void UBehaviorTreeGraphNode::CreateAddServiceSubMenu(UToolMenu* Menu, UEdGraph* Graph) const
 {
-	TSharedRef<SGraphEditorActionMenuAI> Menu =	
+	TSharedRef<SGraphEditorActionMenuAI> Widget =
 		SNew(SGraphEditorActionMenuAI)
 		.GraphObj( Graph )
 		.GraphNode((UBehaviorTreeGraphNode*)this)
 		.SubNodeFlags(ESubNode::Service)
 		.AutoExpandActionMenu(true);
 
-	MenuBuilder.AddWidget(Menu,FText(),true);
+	FToolMenuSection& Section = Menu->FindOrAddSection("Section");
+	Section.AddEntry(FToolMenuEntry::InitWidget("ServiceWidget", Widget, FText(), true));
 }
 
-void UBehaviorTreeGraphNode::AddContextMenuActionsDecorators(const FGraphNodeContextMenuBuilder& Context) const
+void UBehaviorTreeGraphNode::AddContextMenuActionsDecorators(UToolMenu* Menu, const FName SectionName, UGraphNodeContextMenuContext* Context) const
 {
-	Context.MenuBuilder->AddSubMenu(
+	FToolMenuSection& Section = Menu->FindOrAddSection(SectionName);
+	Section.AddSubMenu(
+		"AddDecorator",
 		LOCTEXT("AddDecorator", "Add Decorator..." ),
 		LOCTEXT("AddDecoratorTooltip", "Adds new decorator as a subnode" ),
-		FNewMenuDelegate::CreateUObject( this, &UBehaviorTreeGraphNode::CreateAddDecoratorSubMenu,(UEdGraph*)Context.Graph));
+		FNewToolMenuDelegate::CreateUObject(this, &UBehaviorTreeGraphNode::CreateAddDecoratorSubMenu, (UEdGraph*)Context->Graph));
 }
 
-void UBehaviorTreeGraphNode::AddContextMenuActionsServices(const FGraphNodeContextMenuBuilder& Context) const
+void UBehaviorTreeGraphNode::AddContextMenuActionsServices(UToolMenu* Menu, const FName SectionName, UGraphNodeContextMenuContext* Context) const
 {
-	Context.MenuBuilder->AddSubMenu(
+	FToolMenuSection& Section = Menu->FindOrAddSection(SectionName);
+	Section.AddSubMenu(
+		"AddService",
 		LOCTEXT("AddService", "Add Service..." ),
 		LOCTEXT("AddServiceTooltip", "Adds new service as a subnode" ),
-		FNewMenuDelegate::CreateUObject( this, &UBehaviorTreeGraphNode::CreateAddServiceSubMenu,(UEdGraph*)Context.Graph));
+		FNewToolMenuDelegate::CreateUObject(this, &UBehaviorTreeGraphNode::CreateAddServiceSubMenu, (UEdGraph*)Context->Graph));
 }
 
 void UBehaviorTreeGraphNode::ClearDebuggerState()
diff --git a/Engine/Source/Editor/BehaviorTreeEditor/Private/BehaviorTreeGraphNode_Composite.cpp b/Engine/Source/Editor/BehaviorTreeEditor/Private/BehaviorTreeGraphNode_Composite.cpp
--- a/Engine/Source/Editor/BehaviorTreeEditor/Private/BehaviorTreeGraphNode_Composite.cpp
+++ b/Engine/Source/Editor/BehaviorTreeEditor/Private/BehaviorTreeGraphNode_Composite.cpp
@@ -45,8 +45,8 @@ FText UBehaviorTreeGraphNode_Composite::GetTooltipText() const
 	return Super::GetTooltipText();
 }
 
-void UBehaviorTreeGraphNode_Composite::GetContextMenuActions(const FGraphNodeContextMenuBuilder& Context) const
+void UBehaviorTreeGraphNode_Composite::GetNodeContextMenuActions(UToolMenu* Menu, UGraphNodeContextMenuContext* Context) const
 {
-	AddContextMenuActionsDecorators(Context);
-	AddContextMenuActionsServices(Context);
+	AddContextMenuActionsDecorators(Menu, "BehaviorTreeGraphNode", Context);
+	AddContextMenuActionsServices(Menu, "BehaviorTreeGraphNode", Context);
 }
diff --git a/Engine/Source/Editor/BehaviorTreeEditor/Private/BehaviorTreeGraphNode_Task.cpp b/Engine/Source/Editor/BehaviorTreeEditor/Private/BehaviorTreeGraphNode_Task.cpp
--- a/Engine/Source/Editor/BehaviorTreeEditor/Private/BehaviorTreeGraphNode_Task.cpp
+++ b/Engine/Source/Editor/BehaviorTreeEditor/Private/BehaviorTreeGraphNode_Task.cpp
@@ -31,8 +31,8 @@ FText UBehaviorTreeGraphNode_Task::GetNodeTitle(ENodeTitleType::Type TitleType)
 	return Super::GetNodeTitle(TitleType);
 }
 
-void UBehaviorTreeGraphNode_Task::GetContextMenuActions(const FGraphNodeContextMenuBuilder& Context) const
+void UBehaviorTreeGraphNode_Task::GetNodeContextMenuActions(UToolMenu* Menu, UGraphNodeContextMenuContext* Context) const
 {
-	AddContextMenuActionsDecorators(Context);
-	AddContextMenuActionsServices(Context);
+	AddContextMenuActionsDecorators(Menu, "BehaviorTreeGraphNode", Context);
+	AddContextMenuActionsServices(Menu, "BehaviorTreeGraphNode", Context);
 }
diff --git a/Engine/Source/Editor/BlueprintGraph/Classes/K2Node_CommutativeAssociativeBinaryOperator.h b/Engine/Source/Editor/BlueprintGraph/Classes/K2Node_CommutativeAssociativeBinaryOperator.h
--- a/Engine/Source/Editor/BlueprintGraph/Classes/K2Node_CommutativeAssociativeBinaryOperator.h
+++ b/Engine/Source/Editor/BlueprintGraph/Classes/K2Node_CommutativeAssociativeBinaryOperator.h
@@ -42,10 +42,11 @@ public:
 	// UEdGraphNode interface
 	virtual void AllocateDefaultPins() override;
 	virtual void ValidateNodeDuringCompilation(class FCompilerResultsLog& MessageLog) const override;
+	virtual void GetNodeContextMenuActions(class UToolMenu* Menu, class UGraphNodeContextMenuContext* Context) const override;
+	virtual bool IncludeParentNodeContextMenu() const override { return true; }
 	// End of UEdGraphNode interface
 
 	// UK2Node interface
-	virtual void GetContextMenuActions(const FGraphNodeContextMenuBuilder& Context) const override;
 	virtual void ExpandNode(class FKismetCompilerContext& CompilerContext, UEdGraph* SourceGraph) override;
 	// End of UK2Node interface
 
diff --git a/Engine/Source/Editor/BlueprintGraph/Classes/K2Node_DoOnceMultiInput.h b/Engine/Source/Editor/BlueprintGraph/Classes/K2Node_DoOnceMultiInput.h
--- a/Engine/Source/Editor/BlueprintGraph/Classes/K2Node_DoOnceMultiInput.h
+++ b/Engine/Source/Editor/BlueprintGraph/Classes/K2Node_DoOnceMultiInput.h
@@ -53,10 +53,11 @@ public:
 
 	// UEdGraphNode interface
 	virtual void AllocateDefaultPins() override;
+	virtual void GetNodeContextMenuActions(class UToolMenu* Menu, class UGraphNodeContextMenuContext* Context) const override;
+	virtual bool IncludeParentNodeContextMenu() const override { return true; }
 	// End of UEdGraphNode interface
 
 	// UK2Node interface
-	virtual void GetContextMenuActions(const FGraphNodeContextMenuBuilder& Context) const override;
 	virtual void ExpandNode(class FKismetCompilerContext& CompilerContext, UEdGraph* SourceGraph) override;
 	virtual void GetMenuActions(FBlueprintActionDatabaseRegistrar& ActionRegistrar) const override;
 	virtual FText GetMenuCategory() const override;
diff --git a/Engine/Source/Editor/BlueprintGraph/Classes/K2Node_DynamicCast.h b/Engine/Source/Editor/BlueprintGraph/Classes/K2Node_DynamicCast.h
--- a/Engine/Source/Editor/BlueprintGraph/Classes/K2Node_DynamicCast.h
+++ b/Engine/Source/Editor/BlueprintGraph/Classes/K2Node_DynamicCast.h
@@ -26,7 +26,8 @@ class UK2Node_DynamicCast : public UK2Node
 	virtual FLinearColor GetNodeTitleColor() const override;
 	virtual FText GetNodeTitle(ENodeTitleType::Type TitleType) const override;
 	virtual FSlateIcon GetIconAndTint(FLinearColor& OutColor) const override;
-	virtual void GetContextMenuActions(const FGraphNodeContextMenuBuilder& Context) const override;
+	virtual void GetNodeContextMenuActions(class UToolMenu* Menu, class UGraphNodeContextMenuContext* Context) const override;
+	virtual bool IncludeParentNodeContextMenu() const override { return true; }
 	virtual void PostReconstructNode() override;
 	virtual void PostPlacedNewNode() override;
 	//~ End UEdGraphNode Interface
diff --git a/Engine/Source/Editor/BlueprintGraph/Classes/K2Node_EaseFunction.h b/Engine/Source/Editor/BlueprintGraph/Classes/K2Node_EaseFunction.h
--- a/Engine/Source/Editor/BlueprintGraph/Classes/K2Node_EaseFunction.h
+++ b/Engine/Source/Editor/BlueprintGraph/Classes/K2Node_EaseFunction.h
@@ -41,7 +41,8 @@ class UK2Node_EaseFunction : public UK2Node
 	//~ Begin UEdGraphNode Interface.
 	virtual void AllocateDefaultPins() override;
 	virtual FText GetNodeTitle(ENodeTitleType::Type TitleType) const override;
-	void GetContextMenuActions(const FGraphNodeContextMenuBuilder& Context) const override;
+	virtual void GetNodeContextMenuActions(class UToolMenu* Menu, class UGraphNodeContextMenuContext* Contex) const override;
+	virtual bool IncludeParentNodeContextMenu() const override { return true; }
 	virtual void PinDefaultValueChanged(UEdGraphPin* Pin) override;
 	virtual void PinTypeChanged(UEdGraphPin* Pin) override;
 	virtual FText GetTooltipText() const override;
diff --git a/Engine/Source/Editor/BlueprintGraph/Classes/K2Node_GetArrayItem.h b/Engine/Source/Editor/BlueprintGraph/Classes/K2Node_GetArrayItem.h
--- a/Engine/Source/Editor/BlueprintGraph/Classes/K2Node_GetArrayItem.h
+++ b/Engine/Source/Editor/BlueprintGraph/Classes/K2Node_GetArrayItem.h
@@ -26,7 +26,8 @@ public:
 	virtual FText GetTooltipText() const override;
 	virtual class FNodeHandlingFunctor* CreateNodeHandler(class FKismetCompilerContext& CompilerContext) const override;
 	virtual TSharedPtr<SWidget> CreateNodeImage() const override;
-	virtual void GetContextMenuActions(const FGraphNodeContextMenuBuilder& Context) const override;
+	virtual void GetNodeContextMenuActions(class UToolMenu* Menu, class UGraphNodeContextMenuContext* Context) const override;
+	virtual bool IncludeParentNodeContextMenu() const override { return true; }
 	virtual FSlateIcon GetIconAndTint(FLinearColor& OutColor) const override;
 	// End of UEdGraphNode interface
 
diff --git a/Engine/Source/Editor/BlueprintGraph/Classes/K2Node_MacroInstance.h b/Engine/Source/Editor/BlueprintGraph/Classes/K2Node_MacroInstance.h
--- a/Engine/Source/Editor/BlueprintGraph/Classes/K2Node_MacroInstance.h
+++ b/Engine/Source/Editor/BlueprintGraph/Classes/K2Node_MacroInstance.h
@@ -49,7 +49,7 @@ public:
 	virtual FLinearColor GetNodeTitleColor() const override;
 	virtual FText GetNodeTitle(ENodeTitleType::Type TitleType) const override;
 	virtual bool CanUserDeleteNode() const override { return true; }
-	void GetContextMenuActions(const FGraphNodeContextMenuBuilder& Context) const override;
+	virtual void GetNodeContextMenuActions(class UToolMenu* Menu, class UGraphNodeContextMenuContext* Context) const override;
 	virtual void NodeConnectionListChanged() override;
 	virtual FString GetDocumentationLink() const override;
 	virtual FString GetDocumentationExcerptName() const override;
diff --git a/Engine/Source/Editor/BlueprintGraph/Classes/K2Node_MakeArray.h b/Engine/Source/Editor/BlueprintGraph/Classes/K2Node_MakeArray.h
--- a/Engine/Source/Editor/BlueprintGraph/Classes/K2Node_MakeArray.h
+++ b/Engine/Source/Editor/BlueprintGraph/Classes/K2Node_MakeArray.h
@@ -23,7 +23,8 @@ public:
 	// End of UEdGraphNode interface
 
 	// UK2Node interface
-	virtual void GetContextMenuActions(const FGraphNodeContextMenuBuilder& Context) const override;
+	virtual void GetNodeContextMenuActions(class UToolMenu* Menu, class UGraphNodeContextMenuContext* Context) const override;
+	virtual bool IncludeParentNodeContextMenu() const override { return true; }
 	virtual class FNodeHandlingFunctor* CreateNodeHandler(class FKismetCompilerContext& CompilerContext) const override;
 	virtual FText GetMenuCategory() const override;
 	// End of UK2Node interface
diff --git a/Engine/Source/Editor/BlueprintGraph/Classes/K2Node_MakeMap.h b/Engine/Source/Editor/BlueprintGraph/Classes/K2Node_MakeMap.h
--- a/Engine/Source/Editor/BlueprintGraph/Classes/K2Node_MakeMap.h
+++ b/Engine/Source/Editor/BlueprintGraph/Classes/K2Node_MakeMap.h
@@ -23,7 +23,8 @@ public:
 	// End of UEdGraphNode interface
 
 	// UK2Node interface
-	virtual void GetContextMenuActions(const FGraphNodeContextMenuBuilder& Context) const override;
+	virtual void GetNodeContextMenuActions(class UToolMenu* Menu, class UGraphNodeContextMenuContext* Context) const override;
+	virtual bool IncludeParentNodeContextMenu() const override { return true; }
 	virtual class FNodeHandlingFunctor* CreateNodeHandler(class FKismetCompilerContext& CompilerContext) const override;
 	virtual FText GetMenuCategory() const override;
 	// End of UK2Node interface
diff --git a/Engine/Source/Editor/BlueprintGraph/Classes/K2Node_MakeSet.h b/Engine/Source/Editor/BlueprintGraph/Classes/K2Node_MakeSet.h
--- a/Engine/Source/Editor/BlueprintGraph/Classes/K2Node_MakeSet.h
+++ b/Engine/Source/Editor/BlueprintGraph/Classes/K2Node_MakeSet.h
@@ -23,7 +23,8 @@ public:
 	// End of UEdGraphNode interface
 
 	// UK2Node interface
-	virtual void GetContextMenuActions(const FGraphNodeContextMenuBuilder& Context) const override;
+	virtual void GetNodeContextMenuActions(class UToolMenu* Menu, class UGraphNodeContextMenuContext* Context) const override;
+	virtual bool IncludeParentNodeContextMenu() const override { return true; }
 	virtual class FNodeHandlingFunctor* CreateNodeHandler(class FKismetCompilerContext& CompilerContext) const override;
 	virtual FText GetMenuCategory() const override;
 	// End of UK2Node interface
diff --git a/Engine/Source/Editor/BlueprintGraph/Classes/K2Node_VariableGet.h b/Engine/Source/Editor/BlueprintGraph/Classes/K2Node_VariableGet.h
--- a/Engine/Source/Editor/BlueprintGraph/Classes/K2Node_VariableGet.h
+++ b/Engine/Source/Editor/BlueprintGraph/Classes/K2Node_VariableGet.h
@@ -26,7 +26,8 @@ class BLUEPRINTGRAPH_API UK2Node_VariableGet : public UK2Node_Variable
 	virtual void AllocateDefaultPins() override;
 	virtual FText GetTooltipText() const override;
 	virtual FText GetNodeTitle(ENodeTitleType::Type TitleType) const override;
-	virtual void GetContextMenuActions(const FGraphNodeContextMenuBuilder& Context) const override;
+	virtual void GetNodeContextMenuActions(class UToolMenu* Menu, class UGraphNodeContextMenuContext* Context) const override;
+	virtual bool IncludeParentNodeContextMenu() const override { return true; }
 	virtual void ValidateNodeDuringCompilation(class FCompilerResultsLog& MessageLog) const override;
 	//~ End UEdGraphNode Interface
 
diff --git a/Engine/Source/Editor/BlueprintGraph/Private/K2Node_CommutativeAssociativeBinaryOperator.cpp b/Engine/Source/Editor/BlueprintGraph/Private/K2Node_CommutativeAssociativeBinaryOperator.cpp
--- a/Engine/Source/Editor/BlueprintGraph/Private/K2Node_CommutativeAssociativeBinaryOperator.cpp
+++ b/Engine/Source/Editor/BlueprintGraph/Private/K2Node_CommutativeAssociativeBinaryOperator.cpp
@@ -2,7 +2,7 @@
 
 #include "K2Node_CommutativeAssociativeBinaryOperator.h"
 #include "Framework/Commands/UIAction.h"
-#include "Framework/MultiBox/MultiBoxBuilder.h"
+#include "ToolMenus.h"
 #include "EdGraphSchema_K2.h"
 #include "Kismet2/BlueprintEditorUtils.h"
 
@@ -220,34 +220,35 @@ void UK2Node_CommutativeAssociativeBinaryOperator::RemoveInputPin(UEdGraphPin* P
 	}
 }
 
-void UK2Node_CommutativeAssociativeBinaryOperator::GetContextMenuActions(const FGraphNodeContextMenuBuilder& Context) const
+void UK2Node_CommutativeAssociativeBinaryOperator::GetNodeContextMenuActions(UToolMenu* Menu, UGraphNodeContextMenuContext* Context) const
 {
-	Super::GetContextMenuActions(Context);
+	Super::GetNodeContextMenuActions(Menu, Context);
 
-	if (!Context.bIsDebugging)
+	if (!Context->bIsDebugging)
 	{
 		static FName CommutativeAssociativeBinaryOperatorNodeName = FName("CommutativeAssociativeBinaryOperatorNode");
 		FText CommutativeAssociativeBinaryOperatorStr = LOCTEXT("CommutativeAssociativeBinaryOperatorNode", "Operator Node");
-		if (Context.Pin != NULL)
+		if (Context->Pin != NULL)
 		{
-			if(CanRemovePin(Context.Pin))
+			if(CanRemovePin(Context->Pin))
 			{
-				Context.MenuBuilder->BeginSection(CommutativeAssociativeBinaryOperatorNodeName, CommutativeAssociativeBinaryOperatorStr);
-				Context.MenuBuilder->AddMenuEntry(
+				FToolMenuSection& Section = Menu->AddSection(CommutativeAssociativeBinaryOperatorNodeName, CommutativeAssociativeBinaryOperatorStr);
+				Section.AddMenuEntry(
+					"RemovePin",
 					LOCTEXT("RemovePin", "Remove pin"),
 					LOCTEXT("RemovePinTooltip", "Remove this input pin"),
 					FSlateIcon(),
 					FUIAction(
-						FExecuteAction::CreateUObject(const_cast<UK2Node_CommutativeAssociativeBinaryOperator*>(this), &UK2Node_CommutativeAssociativeBinaryOperator::RemoveInputPin, const_cast<UEdGraphPin*>(Context.Pin))
+						FExecuteAction::CreateUObject(const_cast<UK2Node_CommutativeAssociativeBinaryOperator*>(this), &UK2Node_CommutativeAssociativeBinaryOperator::RemoveInputPin, const_cast<UEdGraphPin*>(Context->Pin))
 					)
 				);
-				Context.MenuBuilder->EndSection();
 			}
 		}
 		else if(CanAddPin())
 		{
-			Context.MenuBuilder->BeginSection(CommutativeAssociativeBinaryOperatorNodeName, CommutativeAssociativeBinaryOperatorStr);
-			Context.MenuBuilder->AddMenuEntry(
+			FToolMenuSection& Section = Menu->AddSection(CommutativeAssociativeBinaryOperatorNodeName, CommutativeAssociativeBinaryOperatorStr);
+			Section.AddMenuEntry(
+				"AddPin",
 				LOCTEXT("AddPin", "Add pin"),
 				LOCTEXT("AddPinTooltip", "Add another input pin"),
 				FSlateIcon(),
@@ -255,7 +256,6 @@ void UK2Node_CommutativeAssociativeBinaryOperator::GetContextMenuActions(const F
 					FExecuteAction::CreateUObject(const_cast<UK2Node_CommutativeAssociativeBinaryOperator*>(this), &UK2Node_CommutativeAssociativeBinaryOperator::AddInputPin)
 				)
 			);
-			Context.MenuBuilder->EndSection();
 		}
 	}
 }
diff --git a/Engine/Source/Editor/BlueprintGraph/Private/K2Node_DoOnceMultiInput.cpp b/Engine/Source/Editor/BlueprintGraph/Private/K2Node_DoOnceMultiInput.cpp
--- a/Engine/Source/Editor/BlueprintGraph/Private/K2Node_DoOnceMultiInput.cpp
+++ b/Engine/Source/Editor/BlueprintGraph/Private/K2Node_DoOnceMultiInput.cpp
@@ -3,7 +3,7 @@
 #include "K2Node_DoOnceMultiInput.h"
 #include "Textures/SlateIcon.h"
 #include "Framework/Commands/UIAction.h"
-#include "Framework/MultiBox/MultiBoxBuilder.h"
+#include "ToolMenus.h"
 #include "EdGraphSchema_K2.h"
 #include "K2Node_AssignmentStatement.h"
 #include "K2Node_IfThenElse.h"
@@ -268,34 +268,35 @@ void UK2Node_DoOnceMultiInput::RemoveInputPin(UEdGraphPin* Pin)
 	}
 }
 
-void UK2Node_DoOnceMultiInput::GetContextMenuActions(const FGraphNodeContextMenuBuilder& Context) const
+void UK2Node_DoOnceMultiInput::GetNodeContextMenuActions(UToolMenu* Menu, UGraphNodeContextMenuContext* Context) const
 {
-	Super::GetContextMenuActions(Context);
+	Super::GetNodeContextMenuActions(Menu, Context);
 
-	if (!Context.bIsDebugging)
+	if (!Context->bIsDebugging)
 	{
 		static FName CommutativeAssociativeBinaryOperatorNodeName = FName("CommutativeAssociativeBinaryOperatorNode");
 		FText CommutativeAssociativeBinaryOperatorStr = LOCTEXT("CommutativeAssociativeBinaryOperatorNode", "Operator Node");
-		if (Context.Pin != NULL)
+		if (Context->Pin != NULL)
 		{
-			if(CanRemovePin(Context.Pin))
+			if(CanRemovePin(Context->Pin))
 			{
-				Context.MenuBuilder->BeginSection(CommutativeAssociativeBinaryOperatorNodeName, CommutativeAssociativeBinaryOperatorStr);
-				Context.MenuBuilder->AddMenuEntry(
+				FToolMenuSection& Section = Menu->AddSection(CommutativeAssociativeBinaryOperatorNodeName, CommutativeAssociativeBinaryOperatorStr);
+				Section.AddMenuEntry(
+					"RemovePin",
 					LOCTEXT("RemovePin", "Remove pin"),
 					LOCTEXT("RemovePinTooltip", "Remove this input pin"),
 					FSlateIcon(),
 					FUIAction(
-						FExecuteAction::CreateUObject(const_cast<UK2Node_DoOnceMultiInput*>(this), &UK2Node_DoOnceMultiInput::RemoveInputPin, const_cast<UEdGraphPin*>(Context.Pin))
+						FExecuteAction::CreateUObject(const_cast<UK2Node_DoOnceMultiInput*>(this), &UK2Node_DoOnceMultiInput::RemoveInputPin, const_cast<UEdGraphPin*>(Context->Pin))
 					)
 				);
-				Context.MenuBuilder->EndSection();
 			}
 		}
 		else if(CanAddPin())
 		{
-			Context.MenuBuilder->BeginSection(CommutativeAssociativeBinaryOperatorNodeName, CommutativeAssociativeBinaryOperatorStr);
-			Context.MenuBuilder->AddMenuEntry(
+			FToolMenuSection& Section = Menu->AddSection(CommutativeAssociativeBinaryOperatorNodeName, CommutativeAssociativeBinaryOperatorStr);
+			Section.AddMenuEntry(
+				"AddPin",
 				LOCTEXT("AddPin", "Add pin"),
 				LOCTEXT("AddPinTooltip", "Add another input pin"),
 				FSlateIcon(),
@@ -303,7 +304,6 @@ void UK2Node_DoOnceMultiInput::GetContextMenuActions(const FGraphNodeContextMenu
 					FExecuteAction::CreateUObject(const_cast<UK2Node_DoOnceMultiInput*>(this), &UK2Node_DoOnceMultiInput::AddInputPin)
 				)
 			);
-			Context.MenuBuilder->EndSection();
 		}
 	}
 }
diff --git a/Engine/Source/Editor/BlueprintGraph/Private/K2Node_DynamicCast.cpp b/Engine/Source/Editor/BlueprintGraph/Private/K2Node_DynamicCast.cpp
--- a/Engine/Source/Editor/BlueprintGraph/Private/K2Node_DynamicCast.cpp
+++ b/Engine/Source/Editor/BlueprintGraph/Private/K2Node_DynamicCast.cpp
@@ -5,7 +5,7 @@
 #include "UObject/Interface.h"
 #include "Engine/Blueprint.h"
 #include "Framework/Commands/UIAction.h"
-#include "Framework/MultiBox/MultiBoxBuilder.h"
+#include "ToolMenus.h"
 #include "EdGraphSchema_K2.h"
 
 #include "BlueprintEditorSettings.h"
@@ -117,13 +117,12 @@ FText UK2Node_DynamicCast::GetNodeTitle(ENodeTitleType::Type TitleType) const
 	return CachedNodeTitle;
 }
 
-void UK2Node_DynamicCast::GetContextMenuActions(const FGraphNodeContextMenuBuilder& Context) const
+void UK2Node_DynamicCast::GetNodeContextMenuActions(UToolMenu* Menu, UGraphNodeContextMenuContext* Context) const
 {
-	Super::GetContextMenuActions(Context);
+	Super::GetNodeContextMenuActions(Menu, Context);
 
-	if (!Context.bIsDebugging)
+	if (!Context->bIsDebugging)
 	{
-		Context.MenuBuilder->BeginSection("K2NodeDynamicCast", LOCTEXT("DynamicCastHeader", "Cast"));
 		{
 			FText MenuEntryTitle = LOCTEXT("MakePureTitle", "Convert to pure cast");
 			FText MenuEntryTooltip = LOCTEXT("MakePureTooltip", "Removes the execution pins to make the node more versatile (NOTE: the cast could still fail, resulting in an invalid output).");
@@ -149,7 +148,9 @@ void UK2Node_DynamicCast::GetContextMenuActions(const FGraphNodeContextMenuBuild
 				}
 			}
 
-			Context.MenuBuilder->AddMenuEntry(
+			FToolMenuSection& Section = Menu->AddSection("K2NodeDynamicCast", LOCTEXT("DynamicCastHeader", "Cast"));
+			Section.AddMenuEntry(
+				"TogglePurity",
 				MenuEntryTitle,
 				MenuEntryTooltip,
 				FSlateIcon(),
@@ -160,7 +161,6 @@ void UK2Node_DynamicCast::GetContextMenuActions(const FGraphNodeContextMenuBuild
 				)
 			);
 		}
-		Context.MenuBuilder->EndSection();
 	}
 }
 
diff --git a/Engine/Source/Editor/BlueprintGraph/Private/K2Node_EaseFunction.cpp b/Engine/Source/Editor/BlueprintGraph/Private/K2Node_EaseFunction.cpp
--- a/Engine/Source/Editor/BlueprintGraph/Private/K2Node_EaseFunction.cpp
+++ b/Engine/Source/Editor/BlueprintGraph/Private/K2Node_EaseFunction.cpp
@@ -2,7 +2,7 @@
 
 #include "K2Node_EaseFunction.h"
 #include "Framework/Commands/UIAction.h"
-#include "Framework/MultiBox/MultiBoxBuilder.h"
+#include "ToolMenus.h"
 #include "EdGraphSchema_K2.h"
 #include "EdGraph/EdGraphNodeUtils.h"
 #include "K2Node_CallFunction.h"
@@ -533,17 +533,18 @@ void UK2Node_EaseFunction::ResetToWildcards()
 	PinTypeChanged(APin);
 }
 
-void UK2Node_EaseFunction::GetContextMenuActions(const FGraphNodeContextMenuBuilder& Context) const
+void UK2Node_EaseFunction::GetNodeContextMenuActions(UToolMenu* Menu, UGraphNodeContextMenuContext* Context) const
 {
-	Super::GetContextMenuActions(Context);
+	Super::GetNodeContextMenuActions(Menu, Context);
 
-	if (!Context.bIsDebugging)
+	if (!Context->bIsDebugging)
 	{
-		if (Context.Pin == NULL)
+		if (Context->Pin == NULL)
 		{
-			Context.MenuBuilder->BeginSection("UK2Node_EaseFunction", LOCTEXT("ContextMenuHeader", "Ease"));
 			{
-				Context.MenuBuilder->AddMenuEntry(
+				FToolMenuSection& Section = Menu->AddSection("UK2Node_EaseFunction", LOCTEXT("ContextMenuHeader", "Ease"));
+				Section.AddMenuEntry(
+					"AddPin",
 					LOCTEXT("AddPin", "Reset to Wildcards"),
 					LOCTEXT("AddPinTooltip", "Resets A, B and Results pins to its default wildcard state"),
 					FSlateIcon(),
@@ -552,7 +553,6 @@ void UK2Node_EaseFunction::GetContextMenuActions(const FGraphNodeContextMenuBuil
 					)
 				);
 			}
-			Context.MenuBuilder->EndSection();
 		}
 	}
 }
diff --git a/Engine/Source/Editor/BlueprintGraph/Private/K2Node_GetArrayItem.cpp b/Engine/Source/Editor/BlueprintGraph/Private/K2Node_GetArrayItem.cpp
--- a/Engine/Source/Editor/BlueprintGraph/Private/K2Node_GetArrayItem.cpp
+++ b/Engine/Source/Editor/BlueprintGraph/Private/K2Node_GetArrayItem.cpp
@@ -16,7 +16,7 @@
 #include "Kismet/KismetArrayLibrary.h" // for Array_Get()
 #include "SPinTypeSelector.h"
 #include "ScopedTransaction.h"
-#include "Framework/MultiBox/MultiBoxBuilder.h" // for FMenuBuilder
+#include "ToolMenus.h" // for UToolMenu
 #include "BlueprintActionFilter.h"
 #include "Widgets/Notifications/SNotificationList.h" // for FNotificationInfo
 #include "Framework/Notifications/NotificationManager.h"
@@ -175,9 +175,9 @@ TSharedPtr<SWidget> UK2Node_GetArrayItem::CreateNodeImage() const
 	return SPinTypeSelector::ConstructPinTypeImage(GetTargetArrayPin());
 }
 
-void UK2Node_GetArrayItem::GetContextMenuActions(const FGraphNodeContextMenuBuilder& Context) const
+void UK2Node_GetArrayItem::GetNodeContextMenuActions(UToolMenu* Menu, UGraphNodeContextMenuContext* Context) const
 {
-	Super::GetContextMenuActions(Context);
+	Super::GetNodeContextMenuActions(Menu, Context);
 
 	const bool bReturnIsRef = IsSetToReturnRef();
 	FText ToggleTooltip = K2Node_GetArrayItem_Impl::GetToggleTooltip(bReturnIsRef);
@@ -189,15 +189,17 @@ void UK2Node_GetArrayItem::GetContextMenuActions(const FGraphNodeContextMenuBuil
 		ToggleTooltip = FText::Format(LOCTEXT("CannotToggleTooltip", "Cannot return by ref using '{0}' pins"), UEdGraphSchema_K2::TypeToText(OutputPin->PinType));
 	}
 
-	Context.MenuBuilder->BeginSection("Array", LOCTEXT("ArrayHeader", "Array Get Node"));
-	Context.MenuBuilder->AddMenuEntry(
-		bReturnIsRef ? LOCTEXT("ChangeNodeToRef", "Change to return a copy") : LOCTEXT("ChangeNodeToVal", "Change to return a reference"),
-		ToggleTooltip,
-		FSlateIcon(),
-		FUIAction(FExecuteAction::CreateUObject(const_cast<UK2Node_GetArrayItem*>(this), &UK2Node_GetArrayItem::ToggleReturnPin),
-			FCanExecuteAction::CreateLambda([bCannotReturnRef]()->bool { return !bCannotReturnRef; }))
+	{
+		FToolMenuSection& Section = Menu->AddSection("Array", LOCTEXT("ArrayHeader", "Array Get Node"));
+		Section.AddMenuEntry(
+			"ToggleReturnPin",
+			bReturnIsRef ? LOCTEXT("ChangeNodeToRef", "Change to return a copy") : LOCTEXT("ChangeNodeToVal", "Change to return a reference"),
+			ToggleTooltip,
+			FSlateIcon(),
+			FUIAction(FExecuteAction::CreateUObject(const_cast<UK2Node_GetArrayItem*>(this), &UK2Node_GetArrayItem::ToggleReturnPin),
+				FCanExecuteAction::CreateLambda([bCannotReturnRef]()->bool { return !bCannotReturnRef; }))
 		);
-	Context.MenuBuilder->EndSection();
+	}
 }
 
 FSlateIcon UK2Node_GetArrayItem::GetIconAndTint(FLinearColor& OutColor) const
diff --git a/Engine/Source/Editor/BlueprintGraph/Private/K2Node_MacroInstance.cpp b/Engine/Source/Editor/BlueprintGraph/Private/K2Node_MacroInstance.cpp
--- a/Engine/Source/Editor/BlueprintGraph/Private/K2Node_MacroInstance.cpp
+++ b/Engine/Source/Editor/BlueprintGraph/Private/K2Node_MacroInstance.cpp
@@ -3,7 +3,7 @@
 #include "K2Node_MacroInstance.h"
 #include "Engine/Blueprint.h"
 #include "Framework/Commands/UIAction.h"
-#include "Framework/MultiBox/MultiBoxBuilder.h"
+#include "ToolMenus.h"
 #include "EdGraphSchema_K2.h"
 #include "Kismet2/BlueprintEditorUtils.h"
 #include "EditorStyleSet.h"
@@ -198,20 +198,20 @@ FLinearColor UK2Node_MacroInstance::GetNodeTitleColor() const
 	return FLinearColor::White;
 }
 
-void UK2Node_MacroInstance::GetContextMenuActions(const FGraphNodeContextMenuBuilder& Context) const
+void UK2Node_MacroInstance::GetNodeContextMenuActions(UToolMenu* Menu, UGraphNodeContextMenuContext* Context) const
 {
-	if ( Context.Pin == nullptr )
+	if ( Context->Pin == nullptr )
 	{
-		Context.MenuBuilder->BeginSection("K2NodeMacroInstance", NSLOCTEXT("K2Node", "MacroInstanceHeader", "Macro Instance"));
 		{
-			Context.MenuBuilder->AddMenuEntry(
+			FToolMenuSection& Section = Menu->AddSection("K2NodeMacroInstance", NSLOCTEXT("K2Node", "MacroInstanceHeader", "Macro Instance"));
+			Section.AddMenuEntry(
+				"MacroInstanceFindInContentBrowser",
 				NSLOCTEXT("K2Node", "MacroInstanceFindInContentBrowser", "Find in Content Browser"),
 				NSLOCTEXT("K2Node", "MacroInstanceFindInContentBrowserTooltip", "Finds the Blueprint Macro Library that contains this Macro in the Content Browser"),
 				FSlateIcon(FEditorStyle::GetStyleSetName(), "PropertyWindow.Button_Browse"),
 				FUIAction( FExecuteAction::CreateStatic( &UK2Node_MacroInstance::FindInContentBrowser, MakeWeakObjectPtr(const_cast<UK2Node_MacroInstance*>(this)) ) )
 				);
 		}
-		Context.MenuBuilder->EndSection();
 	}
 }
 
diff --git a/Engine/Source/Editor/BlueprintGraph/Private/K2Node_MakeArray.cpp b/Engine/Source/Editor/BlueprintGraph/Private/K2Node_MakeArray.cpp
--- a/Engine/Source/Editor/BlueprintGraph/Private/K2Node_MakeArray.cpp
+++ b/Engine/Source/Editor/BlueprintGraph/Private/K2Node_MakeArray.cpp
@@ -5,7 +5,7 @@
 #include "EdGraph/EdGraphPin.h"
 #include "Engine/Blueprint.h"
 #include "Framework/Commands/UIAction.h"
-#include "Framework/MultiBox/MultiBoxBuilder.h"
+#include "ToolMenus.h"
 #include "EdGraphSchema_K2.h"
 #include "EdGraph/EdGraphNodeUtils.h"
 #include "Kismet2/BlueprintEditorUtils.h"
@@ -72,31 +72,33 @@ FSlateIcon UK2Node_MakeArray::GetIconAndTint(FLinearColor& OutColor) const
 	return Icon;
 }
 
-void UK2Node_MakeArray::GetContextMenuActions(const FGraphNodeContextMenuBuilder& Context) const
+void UK2Node_MakeArray::GetNodeContextMenuActions(UToolMenu* Menu, UGraphNodeContextMenuContext* Context) const
 {
-	Super::GetContextMenuActions(Context);
+	Super::GetNodeContextMenuActions(Menu, Context);
 
-	if (!Context.bIsDebugging)
+	if (!Context->bIsDebugging)
 	{
-		Context.MenuBuilder->BeginSection("K2NodeMakeArray", NSLOCTEXT("K2Nodes", "MakeArrayHeader", "MakeArray"));
+		FToolMenuSection& Section = Menu->AddSection("K2NodeMakeArray", NSLOCTEXT("K2Nodes", "MakeArrayHeader", "MakeArray"));
 
-		if (Context.Pin != NULL)
+		if (Context->Pin != NULL)
 		{
-			if (Context.Pin->Direction == EGPD_Input && Context.Pin->ParentPin == nullptr)
+			if (Context->Pin->Direction == EGPD_Input && Context->Pin->ParentPin == nullptr)
 			{
-				Context.MenuBuilder->AddMenuEntry(
+				Section.AddMenuEntry(
+					"RemovePin",
 					LOCTEXT("RemovePin", "Remove array element pin"),
 					LOCTEXT("RemovePinTooltip", "Remove this array element pin"),
 					FSlateIcon(),
 					FUIAction(
-						FExecuteAction::CreateUObject(const_cast<UK2Node_MakeArray*>(this), &UK2Node_MakeArray::RemoveInputPin, const_cast<UEdGraphPin*>(Context.Pin))
+						FExecuteAction::CreateUObject(const_cast<UK2Node_MakeArray*>(this), &UK2Node_MakeArray::RemoveInputPin, const_cast<UEdGraphPin*>(Context->Pin))
 					)
 				);
 			}
 		}
 		else
 		{
-			Context.MenuBuilder->AddMenuEntry(
+			Section.AddMenuEntry(
+				"AddPin",
 				LOCTEXT("AddPin", "Add array element pin"),
 				LOCTEXT("AddPinTooltip", "Add another array element pin"),
 				FSlateIcon(),
@@ -106,7 +108,8 @@ void UK2Node_MakeArray::GetContextMenuActions(const FGraphNodeContextMenuBuilder
 			);
 		}
 
-		Context.MenuBuilder->AddMenuEntry(
+		Section.AddMenuEntry(
+			"ResetToWildcard",
 			LOCTEXT("ResetToWildcard", "Reset to wildcard"),
 			LOCTEXT("ResetToWildcardTooltip", "Reset the node to have wildcard input/outputs. Requires no pins are connected."),
 			FSlateIcon(),
@@ -115,8 +118,6 @@ void UK2Node_MakeArray::GetContextMenuActions(const FGraphNodeContextMenuBuilder
 				FCanExecuteAction::CreateUObject(this, &UK2Node_MakeArray::CanResetToWildcard)
 			)
 		);
-
-		Context.MenuBuilder->EndSection();
 	}
 }
 
diff --git a/Engine/Source/Editor/BlueprintGraph/Private/K2Node_MakeMap.cpp b/Engine/Source/Editor/BlueprintGraph/Private/K2Node_MakeMap.cpp
--- a/Engine/Source/Editor/BlueprintGraph/Private/K2Node_MakeMap.cpp
+++ b/Engine/Source/Editor/BlueprintGraph/Private/K2Node_MakeMap.cpp
@@ -5,7 +5,7 @@
 #include "EdGraph/EdGraphPin.h"
 #include "Engine/Blueprint.h"
 #include "Framework/Commands/UIAction.h"
-#include "Framework/MultiBox/MultiBoxBuilder.h"
+#include "ToolMenus.h"
 #include "EdGraphSchema_K2.h"
 #include "EdGraph/EdGraphNodeUtils.h"
 #include "Kismet2/BlueprintEditorUtils.h"
@@ -144,31 +144,33 @@ FSlateIcon UK2Node_MakeMap::GetIconAndTint(FLinearColor& OutColor) const
 	return Icon;
 }
 
-void UK2Node_MakeMap::GetContextMenuActions(const FGraphNodeContextMenuBuilder& Context) const
+void UK2Node_MakeMap::GetNodeContextMenuActions(UToolMenu* Menu, UGraphNodeContextMenuContext* Context) const
 {
-	Super::GetContextMenuActions(Context);
+	Super::GetNodeContextMenuActions(Menu, Context);
 
-	if (!Context.bIsDebugging)
+	if (!Context->bIsDebugging)
 	{
-		Context.MenuBuilder->BeginSection("K2NodeMakeMap", NSLOCTEXT("K2Nodes", "MakeMapHeader", "MakeMap"));
+		FToolMenuSection& Section = Menu->AddSection("K2NodeMakeMap", NSLOCTEXT("K2Nodes", "MakeMapHeader", "MakeMap"));
 
-		if (Context.Pin)
+		if (Context->Pin)
 		{
-			if (Context.Pin->Direction == EGPD_Input && Context.Pin->ParentPin == nullptr)
+			if (Context->Pin->Direction == EGPD_Input && Context->Pin->ParentPin == nullptr)
 			{
-				Context.MenuBuilder->AddMenuEntry(
+				Section.AddMenuEntry(
+					"RemovePin",
 					LOCTEXT("RemovePin", "Remove key/value pair"),
 					LOCTEXT("RemovePinTooltip", "Remove this pin and its corresponding key/value pin"),
 					FSlateIcon(),
 					FUIAction(
-						FExecuteAction::CreateUObject(const_cast<UK2Node_MakeMap*>(this), &UK2Node_MakeMap::RemoveInputPin, const_cast<UEdGraphPin*>(Context.Pin))
+						FExecuteAction::CreateUObject(const_cast<UK2Node_MakeMap*>(this), &UK2Node_MakeMap::RemoveInputPin, const_cast<UEdGraphPin*>(Context->Pin))
 					)
 				);
 			}
 		}
 		else
 		{
-			Context.MenuBuilder->AddMenuEntry(
+			Section.AddMenuEntry(
+				"AddPin",
 				LOCTEXT("AddPin", "Add key/value pair"),
 				LOCTEXT("AddPinTooltip", "Add another pair of key/value pins"),
 				FSlateIcon(),
@@ -178,7 +180,8 @@ void UK2Node_MakeMap::GetContextMenuActions(const FGraphNodeContextMenuBuilder&
 			);
 		}
 
-		Context.MenuBuilder->AddMenuEntry(
+		Section.AddMenuEntry(
+			"ResetToWildcard",
 			LOCTEXT("ResetToWildcard", "Reset to wildcard"),
 			LOCTEXT("ResetToWildcardTooltip", "Reset the node to have wildcard input/outputs. Requires no pins are connected."),
 			FSlateIcon(),
@@ -187,8 +190,6 @@ void UK2Node_MakeMap::GetContextMenuActions(const FGraphNodeContextMenuBuilder&
 				FCanExecuteAction::CreateUObject(this, &UK2Node_MakeMap::CanResetToWildcard)
 			)
 		);
-
-		Context.MenuBuilder->EndSection();
 	}
 }
 
diff --git a/Engine/Source/Editor/BlueprintGraph/Private/K2Node_MakeSet.cpp b/Engine/Source/Editor/BlueprintGraph/Private/K2Node_MakeSet.cpp
--- a/Engine/Source/Editor/BlueprintGraph/Private/K2Node_MakeSet.cpp
+++ b/Engine/Source/Editor/BlueprintGraph/Private/K2Node_MakeSet.cpp
@@ -5,7 +5,7 @@
 #include "EdGraph/EdGraphPin.h"
 #include "Engine/Blueprint.h"
 #include "Framework/Commands/UIAction.h"
-#include "Framework/MultiBox/MultiBoxBuilder.h"
+#include "ToolMenus.h"
 #include "EdGraphSchema_K2.h"
 #include "EdGraph/EdGraphNodeUtils.h"
 #include "Kismet2/BlueprintEditorUtils.h"
@@ -72,31 +72,33 @@ FSlateIcon UK2Node_MakeSet::GetIconAndTint(FLinearColor& OutColor) const
 	return Icon;
 }
 
-void UK2Node_MakeSet::GetContextMenuActions(const FGraphNodeContextMenuBuilder& Context) const
+void UK2Node_MakeSet::GetNodeContextMenuActions(UToolMenu* Menu, UGraphNodeContextMenuContext* Context) const
 {
-	Super::GetContextMenuActions(Context);
+	Super::GetNodeContextMenuActions(Menu, Context);
 
-	if (!Context.bIsDebugging)
+	if (!Context->bIsDebugging)
 	{
-		Context.MenuBuilder->BeginSection("K2NodeMakeSet", NSLOCTEXT("K2Nodes", "MakeSetHeader", "MakeSet"));
+		FToolMenuSection& Section = Menu->AddSection("K2NodeMakeSet", NSLOCTEXT("K2Nodes", "MakeSetHeader", "MakeSet"));
 
-		if (Context.Pin)
+		if (Context->Pin)
 		{
-			if (Context.Pin->Direction == EGPD_Input && Context.Pin->ParentPin == nullptr)
+			if (Context->Pin->Direction == EGPD_Input && Context->Pin->ParentPin == nullptr)
 			{
-				Context.MenuBuilder->AddMenuEntry(
+				Section.AddMenuEntry(
+					"RemovePin",
 					LOCTEXT("RemovePin", "Remove set element pin"),
 					LOCTEXT("RemovePinTooltip", "Remove this set element pin"),
 					FSlateIcon(),
 					FUIAction(
-						FExecuteAction::CreateUObject(const_cast<UK2Node_MakeSet*>(this), &UK2Node_MakeSet::RemoveInputPin, const_cast<UEdGraphPin*>(Context.Pin))
+						FExecuteAction::CreateUObject(const_cast<UK2Node_MakeSet*>(this), &UK2Node_MakeSet::RemoveInputPin, const_cast<UEdGraphPin*>(Context->Pin))
 					)
 				);
 			}
 		}
 		else
 		{
-			Context.MenuBuilder->AddMenuEntry(
+			Section.AddMenuEntry(
+				"AddPin",
 				LOCTEXT("AddPin", "Add set element pin"),
 				LOCTEXT("AddPinTooltip", "Add another set element pin"),
 				FSlateIcon(),
@@ -106,7 +108,8 @@ void UK2Node_MakeSet::GetContextMenuActions(const FGraphNodeContextMenuBuilder&
 			);
 		}
 
-		Context.MenuBuilder->AddMenuEntry(
+		Section.AddMenuEntry(
+			"ResetToWildcard",
 			LOCTEXT("ResetToWildcard", "Reset to wildcard"),
 			LOCTEXT("ResetToWildcardTooltip", "Reset the node to have wildcard input/outputs. Requires no pins are connected."),
 			FSlateIcon(),
@@ -115,8 +118,6 @@ void UK2Node_MakeSet::GetContextMenuActions(const FGraphNodeContextMenuBuilder&
 				FCanExecuteAction::CreateUObject(this, &UK2Node_MakeSet::CanResetToWildcard)
 			)
 		);
-
-		Context.MenuBuilder->EndSection();
 	}
 }
 
diff --git a/Engine/Source/Editor/BlueprintGraph/Private/K2Node_VariableGet.cpp b/Engine/Source/Editor/BlueprintGraph/Private/K2Node_VariableGet.cpp
--- a/Engine/Source/Editor/BlueprintGraph/Private/K2Node_VariableGet.cpp
+++ b/Engine/Source/Editor/BlueprintGraph/Private/K2Node_VariableGet.cpp
@@ -6,7 +6,7 @@
 #include "UObject/PropertyPortFlags.h"
 #include "Kismet/KismetSystemLibrary.h"
 #include "Framework/Commands/UIAction.h"
-#include "Framework/MultiBox/MultiBoxBuilder.h"
+#include "ToolMenus.h"
 #include "EdGraph/EdGraph.h"
 #include "EdGraphSchema_K2.h"
 #include "K2Node_CallFunction.h"
@@ -318,14 +318,13 @@ bool UK2Node_VariableGet::IsValidTypeForNonPure(const FEdGraphPinType& InPinType
 	return !InPinType.IsContainer() && (InPinType.PinCategory == UEdGraphSchema_K2::PC_Object || InPinType.PinCategory == UEdGraphSchema_K2::PC_Class || InPinType.PinCategory == UEdGraphSchema_K2::PC_SoftObject || InPinType.PinCategory == UEdGraphSchema_K2::PC_SoftClass);
 }
 
-void UK2Node_VariableGet::GetContextMenuActions(const FGraphNodeContextMenuBuilder& Context) const
+void UK2Node_VariableGet::GetNodeContextMenuActions(UToolMenu* Menu, UGraphNodeContextMenuContext* Context) const
 {
-	Super::GetContextMenuActions(Context);
+	Super::GetNodeContextMenuActions(Menu, Context);
 
 	const UEdGraphPin* ValuePin = GetValuePin();
 	if (ValuePin && IsValidTypeForNonPure(ValuePin->PinType))
 	{
-		Context.MenuBuilder->BeginSection("K2NodeVariableGet", LOCTEXT("VariableGetHeader", "Variable Get"));
 		{
 			FText MenuEntryTitle;
 			FText MenuEntryTooltip;
@@ -356,18 +355,19 @@ void UK2Node_VariableGet::GetContextMenuActions(const FGraphNodeContextMenuBuild
 				MenuEntryTooltip = LOCTEXT("ConvertToPureGetTooltip", "Removes the execution pins to make the node more versatile.");
 			}
 
-			Context.MenuBuilder->AddMenuEntry(
+			FToolMenuSection& Section = Menu->AddSection("K2NodeVariableGet", LOCTEXT("VariableGetHeader", "Variable Get"));
+			Section.AddMenuEntry(
+				"TogglePurity",
 				MenuEntryTitle,
 				MenuEntryTooltip,
 				FSlateIcon(),
 				FUIAction(
 				FExecuteAction::CreateUObject(const_cast<UK2Node_VariableGet*>(this), &UK2Node_VariableGet::TogglePurity),
-				FCanExecuteAction::CreateStatic(CanExecutePurityToggle, bCanTogglePurity && !Context.bIsDebugging),
+				FCanExecuteAction::CreateStatic(CanExecutePurityToggle, bCanTogglePurity && !Context->bIsDebugging),
 				FIsActionChecked()
 				)
 				);
 		}
-		Context.MenuBuilder->EndSection();
 	}
 }
 
diff --git a/Engine/Source/Editor/EnvironmentQueryEditor/Classes/EnvironmentQueryGraphNode_Option.h b/Engine/Source/Editor/EnvironmentQueryEditor/Classes/EnvironmentQueryGraphNode_Option.h
--- a/Engine/Source/Editor/EnvironmentQueryEditor/Classes/EnvironmentQueryGraphNode_Option.h
+++ b/Engine/Source/Editor/EnvironmentQueryEditor/Classes/EnvironmentQueryGraphNode_Option.h
@@ -26,8 +26,8 @@ class UEnvironmentQueryGraphNode_Option : public UEnvironmentQueryGraphNode
 	virtual void PrepareForCopying() override;
 	virtual void UpdateNodeClassData() override;
 
-	virtual void GetContextMenuActions(const FGraphNodeContextMenuBuilder& Context) const override;
-	void CreateAddTestSubMenu(class FMenuBuilder& MenuBuilder, UEdGraph* Graph) const;
+	virtual void GetNodeContextMenuActions(class UToolMenu* Menu, class UGraphNodeContextMenuContext* Context) const override;
+	void CreateAddTestSubMenu(class UToolMenu* Menu, UEdGraph* Graph) const;
 
 	void CalculateWeights();
 	void UpdateNodeData();
diff --git a/Engine/Source/Editor/EnvironmentQueryEditor/EnvironmentQueryEditor.Build.cs b/Engine/Source/Editor/EnvironmentQueryEditor/EnvironmentQueryEditor.Build.cs
--- a/Engine/Source/Editor/EnvironmentQueryEditor/EnvironmentQueryEditor.Build.cs
+++ b/Engine/Source/Editor/EnvironmentQueryEditor/EnvironmentQueryEditor.Build.cs
@@ -42,6 +42,7 @@ public class EnvironmentQueryEditor : ModuleRules
 				"BlueprintGraph",
                 "AIGraph",
                 "AIModule",
+				"ToolMenus",
 			}
 		);
 
diff --git a/Engine/Source/Editor/EnvironmentQueryEditor/Private/EnvironmentQueryGraphNode_Option.cpp b/Engine/Source/Editor/EnvironmentQueryEditor/Private/EnvironmentQueryGraphNode_Option.cpp
--- a/Engine/Source/Editor/EnvironmentQueryEditor/Private/EnvironmentQueryGraphNode_Option.cpp
+++ b/Engine/Source/Editor/EnvironmentQueryEditor/Private/EnvironmentQueryGraphNode_Option.cpp
@@ -2,7 +2,7 @@
 
 #include "EnvironmentQueryGraphNode_Option.h"
 #include "Widgets/DeclarativeSyntaxSupport.h"
-#include "Framework/MultiBox/MultiBoxBuilder.h"
+#include "ToolMenus.h"
 #include "EnvironmentQuery/EnvQueryGenerator.h"
 #include "EnvironmentQuery/EnvQueryTest.h"
 #include "EnvironmentQuery/EnvQueryOption.h"
@@ -93,24 +93,26 @@ FText UEnvironmentQueryGraphNode_Option::GetDescription() const
 	return OptionInstance ? OptionInstance->GetDescriptionDetails() : FText::GetEmpty();
 }
 
-void UEnvironmentQueryGraphNode_Option::GetContextMenuActions(const FGraphNodeContextMenuBuilder& Context) const
+void UEnvironmentQueryGraphNode_Option::GetNodeContextMenuActions(UToolMenu* Menu, UGraphNodeContextMenuContext* Context) const
 {
-	Context.MenuBuilder->AddSubMenu(
+	FToolMenuSection& Section = Menu->AddSection("EnvironmentQueryGraphNode");
+	Section.AddSubMenu(
+		"AddTest",
 		LOCTEXT("AddTest", "Add Test..." ),
 		LOCTEXT("AddTestTooltip", "Adds new test to generator" ),
-		FNewMenuDelegate::CreateUObject( this, &UEnvironmentQueryGraphNode_Option::CreateAddTestSubMenu,(UEdGraph*)Context.Graph) 
+		FNewToolMenuDelegate::CreateUObject( this, &UEnvironmentQueryGraphNode_Option::CreateAddTestSubMenu,(UEdGraph*)Context->Graph) 
 		);
 }
 
-void UEnvironmentQueryGraphNode_Option::CreateAddTestSubMenu(class FMenuBuilder& MenuBuilder, UEdGraph* Graph) const
+void UEnvironmentQueryGraphNode_Option::CreateAddTestSubMenu(UToolMenu* Menu, UEdGraph* Graph) const
 {
-	TSharedRef<SGraphEditorActionMenuAI> Menu =	
+	TSharedRef<SGraphEditorActionMenuAI> Widget =	
 		SNew(SGraphEditorActionMenuAI)
 		.GraphObj( Graph )
 		.GraphNode((UEnvironmentQueryGraphNode_Option*)this)
 		.AutoExpandActionMenu(true);
 
-	MenuBuilder.AddWidget(Menu,FText(),true);
+	Menu->AddMenuEntry("Section", FToolMenuEntry::InitWidget("Widget", Widget, FText(), true));
 }
 
 void UEnvironmentQueryGraphNode_Option::CalculateWeights()
diff --git a/Engine/Source/Editor/GraphEditor/Private/SGraphEditorImpl.cpp b/Engine/Source/Editor/GraphEditor/Private/SGraphEditorImpl.cpp
--- a/Engine/Source/Editor/GraphEditor/Private/SGraphEditorImpl.cpp
+++ b/Engine/Source/Editor/GraphEditor/Private/SGraphEditorImpl.cpp
@@ -513,21 +513,13 @@ void SGraphEditorImpl::RegisterContextMenuFor_EdGraphSchema()
 
 	Menu->AddDynamicSection("GetNodeContextMenuActions", FNewToolMenuDelegate::CreateLambda([](UToolMenu* InMenu)
 	{
-		if (UGraphNodeContextMenuContext* Context = InMenu->FindContext<UGraphNodeContextMenuContext>())
+		UGraphNodeContextMenuContext* Context = InMenu->FindContext<UGraphNodeContextMenuContext>();
+		if (Context && Context->Node)
 		{
 			Context->Node->GetNodeContextMenuActions(InMenu, Context);
 		}
 	}));
 
-	Menu->AddDynamicSection("GetNodeContextMenuActionsLegacy", FNewToolMenuDelegateLegacy::CreateLambda([](FMenuBuilder& InMenuBuilder, UToolMenu* InMenu)
-	{
-		if (UGraphNodeContextMenuContext* Context = InMenu->FindContext<UGraphNodeContextMenuContext>())
-		{
-			FGraphNodeContextMenuBuilder MenuBuilderContext(Context->Graph, Context->Node, Context->Pin, &InMenuBuilder, Context->bIsDebugging);
-			Context->Node->GetContextMenuActions(MenuBuilderContext);
-		}
-	}));
-
 	Menu->AddDynamicSection("EdGraphSchema", FNewToolMenuDelegate::CreateLambda([](UToolMenu* InMenu)
 	{
 		UGraphNodeContextMenuContext* Context = InMenu->FindContext<UGraphNodeContextMenuContext>();
@@ -656,8 +648,31 @@ void SGraphEditorImpl::RegisterContextMenuFor_EdGraphSchema()
 	}));
 }
 
+FName SGraphEditorImpl::GetNodeParentContextMenuName(UClass* InClass)
+{
+	if (InClass && InClass != UEdGraphNode::StaticClass())
+	{
+		if (InClass->GetDefaultObject<UEdGraphNode>()->IncludeParentNodeContextMenu())
+		{
+			if (UClass* SuperClass = InClass->GetSuperClass())
+			{
+				return GetNodeContextMenuName(SuperClass);
+			}
+		}
+	}
+
+	return NAME_None;
+}
+
+FName SGraphEditorImpl::GetNodeContextMenuName(UClass* InClass)
+{
+	return FName(*(FString(TEXT("GraphEditor.GraphNodeContextMenu.")) + InClass->GetName()));
+}
+
 UToolMenu* SGraphEditorImpl::GenerateContextMenu(const UEdGraphSchema* Schema, FToolMenuContext& MenuContext) const
 {
+	UGraphNodeContextMenuContext* Context = MenuContext.Find<UGraphNodeContextMenuContext>();
+
 	// Register UEdGraphSchema's menu here to reduce amount of editor code in engine module
 	RegisterContextMenuFor_EdGraphSchema();
 
@@ -671,19 +686,47 @@ UToolMenu* SGraphEditorImpl::GenerateContextMenu(const UEdGraphSchema* Schema, F
 		{
 			ToolMenus->RegisterMenu(CheckMenuName, CurrentSchema->GetParentContextMenuName());
 		}
-
-		CurrentClass = CurrentClass->GetSuperClass();
 	}
 
 	const FName MenuName = Schema->GetContextMenuName();
 	TArray<UToolMenu*> Hierarchy = ToolMenus->CollectHierarchy(MenuName);
 
+	// Insert list of node specific menus
+	if (Context->Node)
+	{
+		// Only insert list if graph schema includes actions for nodes
+		static const FName EdGraphSchemaMenuName = UEdGraphSchema::GetContextMenuName(UEdGraphSchema::StaticClass());
+		if (Hierarchy.ContainsByPredicate([=](UToolMenu* OtherMenu) { return OtherMenu->MenuName == EdGraphSchemaMenuName; }))
+		{
+			// Walk class hierarchy
+			for (UClass* CurrentClass = Context->Node->GetClass(); CurrentClass && CurrentClass->IsChildOf(UEdGraphNode::StaticClass()); CurrentClass = CurrentClass->GetSuperClass())
+			{
+				const FName CheckMenuName = GetNodeContextMenuName(CurrentClass);
+				const FName CheckParentName = GetNodeParentContextMenuName(CurrentClass);
+
+				// Register if not already registered
+				if (!ToolMenus->IsMenuRegistered(CheckMenuName))
+				{
+					ToolMenus->RegisterMenu(CheckMenuName, CheckParentName);
+				}
+
+				// Parent rest of hierarchy to this menu
+				Hierarchy.Insert(ToolMenus->FindMenu(CheckMenuName), 0);
+
+				// Stop walking class hierarchy if IncludeParentNodeContextMenu() returns false for current class
+				if (CheckParentName == NAME_None)
+				{
+					break;
+				}
+			}
+		}
+	}
+
 	// Add extra menu into hierarchy for calling GetContextMenuActions()
 	const FName CommonRootMenuName = "GraphEditor.GraphContextMenu.Common";
-	UToolMenu* CommonRootMenu = nullptr;
 	if (!ToolMenus->IsMenuRegistered(CommonRootMenuName))
 	{
-		CommonRootMenu = ToolMenus->RegisterMenu(CommonRootMenuName);
+		UToolMenu* CommonRootMenu = ToolMenus->RegisterMenu(CommonRootMenuName);
 		CommonRootMenu->AddDynamicSection("GetContextMenuActions", FNewToolMenuDelegate::CreateLambda([](UToolMenu* InMenu)
 		{
 			if (UGraphNodeContextMenuContext* ContextObject = InMenu->FindContext<UGraphNodeContextMenuContext>())
@@ -695,11 +738,7 @@ UToolMenu* SGraphEditorImpl::GenerateContextMenu(const UEdGraphSchema* Schema, F
 			}
 		}));
 	}
-	else
-	{
-		CommonRootMenu = ToolMenus->FindMenu(CommonRootMenuName);
-	}
-	Hierarchy.Insert(CommonRootMenu, 0);
+	Hierarchy.Insert(ToolMenus->FindMenu(CommonRootMenuName), 0);
 
 	return ToolMenus->GenerateMenu(Hierarchy, MenuContext);
 }
diff --git a/Engine/Source/Editor/GraphEditor/Private/SGraphEditorImpl.h b/Engine/Source/Editor/GraphEditor/Private/SGraphEditorImpl.h
--- a/Engine/Source/Editor/GraphEditor/Private/SGraphEditorImpl.h
+++ b/Engine/Source/Editor/GraphEditor/Private/SGraphEditorImpl.h
@@ -337,6 +337,8 @@ private:
 	class UToolMenu* GenerateContextMenu(const class UEdGraphSchema* Schema, struct FToolMenuContext& MenuContext) const;
 
 	static void RegisterContextMenuFor_EdGraphSchema();
+	static FName GetNodeParentContextMenuName(UClass* InClass);
+	static FName GetNodeContextMenuName(UClass* InClass);
 
 private:
 	bool bIsActiveTimerRegistered;
diff --git a/Engine/Source/Editor/MovieSceneTools/MovieSceneTools.Build.cs b/Engine/Source/Editor/MovieSceneTools/MovieSceneTools.Build.cs
--- a/Engine/Source/Editor/MovieSceneTools/MovieSceneTools.Build.cs
+++ b/Engine/Source/Editor/MovieSceneTools/MovieSceneTools.Build.cs
@@ -59,6 +59,7 @@ public class MovieSceneTools : ModuleRules
                 "AnimationCore",
 				"TimeManagement",
                 "XmlParser",
+				"ToolMenus",
 			}
 		);
 
diff --git a/Engine/Source/Editor/MovieSceneTools/Private/K2Node_GetSequenceBinding.cpp b/Engine/Source/Editor/MovieSceneTools/Private/K2Node_GetSequenceBinding.cpp
--- a/Engine/Source/Editor/MovieSceneTools/Private/K2Node_GetSequenceBinding.cpp
+++ b/Engine/Source/Editor/MovieSceneTools/Private/K2Node_GetSequenceBinding.cpp
@@ -9,7 +9,7 @@
 #include "Framework/Application/SlateApplication.h"
 #include "PropertyCustomizationHelpers.h"
 #include "MovieSceneSequence.h"
-#include "Framework/MultiBox/MultiBoxBuilder.h"
+#include "ToolMenus.h"
 #include "MovieSceneObjectBindingIDPicker.h"
 #include "SGraphNode.h"
 #include "ContentBrowserModule.h"
@@ -239,18 +239,22 @@ FSlateIcon UK2Node_GetSequenceBinding::GetIconAndTint(FLinearColor& OutColor) co
 	return Icon;
 }
 
-void UK2Node_GetSequenceBinding::GetContextMenuActions(const FGraphNodeContextMenuBuilder& Context) const
+void UK2Node_GetSequenceBinding::GetNodeContextMenuActions(UToolMenu* Menu, UGraphNodeContextMenuContext* Context) const
 {
-	Super::GetContextMenuActions(Context);
+	Super::GetNodeContextMenuActions(Menu, Context);
 
-	if (!Context.bIsDebugging)
+	if (!Context->bIsDebugging)
 	{
-		Context.MenuBuilder->BeginSection("K2NodeGetSequenceBinding", LOCTEXT("ThisNodeHeader", "This Node"));
-		if (!Context.Pin)
+		FToolMenuSection& Section = Menu->AddSection("K2NodeGetSequenceBinding", LOCTEXT("ThisNodeHeader", "This Node"));
+		if (!Context->Pin)
 		{
 			UMovieSceneSequence* Sequence = GetSequence();
 
-			auto SubMenu = [=](FMenuBuilder& SubMenuBuilder)
+			Section.AddSubMenu(
+				"SetSequence",
+				LOCTEXT("SetSequence_Text", "Sequence"),
+				LOCTEXT("SetSequence_ToolTip", "Sets the sequence to get a binding from"),
+				FNewToolMenuDelegate::CreateLambda([=](UToolMenu* SubMenu)
 				{
 					TArray<const UClass*> AllowedClasses({ UMovieSceneSequence::StaticClass() });
 
@@ -263,17 +267,9 @@ void UK2Node_GetSequenceBinding::GetContextMenuActions(const FGraphNodeContextMe
 						FOnAssetSelected::CreateUObject(const_cast<UK2Node_GetSequenceBinding*>(this), &UK2Node_GetSequenceBinding::SetSequence),
 						FSimpleDelegate());
 					
-					SubMenuBuilder.AddWidget(MenuContent, FText::GetEmpty(), false);
-				};
-
-			Context.MenuBuilder->AddSubMenu(
-				LOCTEXT("SetSequence_Text", "Sequence"),
-				LOCTEXT("SetSequence_ToolTip", "Sets the sequence to get a binding from"),
-				FNewMenuDelegate::CreateLambda(SubMenu)
-				);
+					SubMenu->AddMenuEntry("Section", FToolMenuEntry::InitWidget("Widget", MenuContent, FText::GetEmpty(), false));
+				}));
 		}
-
-		Context.MenuBuilder->EndSection();
 	}
 }
 
diff --git a/Engine/Source/Editor/MovieSceneTools/Public/K2Node_GetSequenceBinding.h b/Engine/Source/Editor/MovieSceneTools/Public/K2Node_GetSequenceBinding.h
--- a/Engine/Source/Editor/MovieSceneTools/Public/K2Node_GetSequenceBinding.h
+++ b/Engine/Source/Editor/MovieSceneTools/Public/K2Node_GetSequenceBinding.h
@@ -41,7 +41,8 @@ public:
 	virtual FSlateIcon GetIconAndTint(FLinearColor& OutColor) const override;
 	virtual bool ShouldShowNodeProperties() const override { return true; }
 	virtual bool IsNodePure() const override { return true; }
-	virtual void GetContextMenuActions(const FGraphNodeContextMenuBuilder& Context) const override;
+	virtual void GetNodeContextMenuActions(class UToolMenu* Menu, class UGraphNodeContextMenuContext* Context) const override;
+	virtual bool IncludeParentNodeContextMenu() const override { return true; }
 	virtual class FNodeHandlingFunctor* CreateNodeHandler(class FKismetCompilerContext& CompilerContext) const override;
 	virtual void PreloadRequiredAssets() override;
 	virtual void GetMenuActions(FBlueprintActionDatabaseRegistrar& ActionRegistrar) const override;
diff --git a/Engine/Source/Editor/UnrealEd/Classes/MaterialGraph/MaterialGraphNode_Comment.h b/Engine/Source/Editor/UnrealEd/Classes/MaterialGraph/MaterialGraphNode_Comment.h
--- a/Engine/Source/Editor/UnrealEd/Classes/MaterialGraph/MaterialGraphNode_Comment.h
+++ b/Engine/Source/Editor/UnrealEd/Classes/MaterialGraph/MaterialGraphNode_Comment.h
@@ -32,7 +32,7 @@ class UMaterialGraphNode_Comment : public UEdGraphNode_Comment
 
 	//~ Begin UEdGraphNode Interface.
 	virtual void PrepareForCopying() override;
-	virtual void GetContextMenuActions(const FGraphNodeContextMenuBuilder& Context) const override;
+	virtual void GetNodeContextMenuActions(class UToolMenu* Menu, class UGraphNodeContextMenuContext* Context) const override;
 	virtual bool CanCreateUnderSpecifiedSchema(const UEdGraphSchema* Schema) const override;
 	virtual void PostPlacedNewNode() override;
 	virtual void OnRenameNode(const FString& NewName) override;
diff --git a/Engine/Source/Editor/UnrealEd/Private/MaterialGraphNode.cpp b/Engine/Source/Editor/UnrealEd/Private/MaterialGraphNode.cpp
--- a/Engine/Source/Editor/UnrealEd/Private/MaterialGraphNode.cpp
+++ b/Engine/Source/Editor/UnrealEd/Private/MaterialGraphNode.cpp
@@ -433,25 +433,29 @@ void UMaterialGraphNode::GetNodeContextMenuActions(UToolMenu* Menu, UGraphNodeCo
 			
 		{
 			FToolMenuSection& Section = Menu->AddSection("Alignment");
-			Section.AddSubMenu(Menu->GetMenuName(), "Alignment", LOCTEXT("AlignmentHeader", "Alignment"), FText(), FNewToolMenuDelegate::CreateLambda([](UToolMenu* InMenu)
-			{
+			Section.AddSubMenu(
+				"Alignment",
+				LOCTEXT("AlignmentHeader", "Alignment"),
+				FText(),
+				FNewToolMenuDelegate::CreateLambda([](UToolMenu* InMenu)
 				{
-					FToolMenuSection& SubMenuSection = InMenu->AddSection("EdGraphSchemaAlignment", LOCTEXT("AlignHeader", "Align"));
-					SubMenuSection.AddMenuEntry(FGraphEditorCommands::Get().AlignNodesTop);
-					SubMenuSection.AddMenuEntry(FGraphEditorCommands::Get().AlignNodesMiddle);
-					SubMenuSection.AddMenuEntry(FGraphEditorCommands::Get().AlignNodesBottom);
-					SubMenuSection.AddMenuEntry(FGraphEditorCommands::Get().AlignNodesLeft);
-					SubMenuSection.AddMenuEntry(FGraphEditorCommands::Get().AlignNodesCenter);
-					SubMenuSection.AddMenuEntry(FGraphEditorCommands::Get().AlignNodesRight);
-					SubMenuSection.AddMenuEntry(FGraphEditorCommands::Get().StraightenConnections);
-				}
+					{
+						FToolMenuSection& SubMenuSection = InMenu->AddSection("EdGraphSchemaAlignment", LOCTEXT("AlignHeader", "Align"));
+						SubMenuSection.AddMenuEntry(FGraphEditorCommands::Get().AlignNodesTop);
+						SubMenuSection.AddMenuEntry(FGraphEditorCommands::Get().AlignNodesMiddle);
+						SubMenuSection.AddMenuEntry(FGraphEditorCommands::Get().AlignNodesBottom);
+						SubMenuSection.AddMenuEntry(FGraphEditorCommands::Get().AlignNodesLeft);
+						SubMenuSection.AddMenuEntry(FGraphEditorCommands::Get().AlignNodesCenter);
+						SubMenuSection.AddMenuEntry(FGraphEditorCommands::Get().AlignNodesRight);
+						SubMenuSection.AddMenuEntry(FGraphEditorCommands::Get().StraightenConnections);
+					}
 
-				{
-					FToolMenuSection& SubMenuSection = InMenu->AddSection("EdGraphSchemaDistribution", LOCTEXT("DistributionHeader", "Distribution"));
-					SubMenuSection.AddMenuEntry(FGraphEditorCommands::Get().DistributeNodesHorizontally);
-					SubMenuSection.AddMenuEntry(FGraphEditorCommands::Get().DistributeNodesVertically);
-				}
-			}));
+					{
+						FToolMenuSection& SubMenuSection = InMenu->AddSection("EdGraphSchemaDistribution", LOCTEXT("DistributionHeader", "Distribution"));
+						SubMenuSection.AddMenuEntry(FGraphEditorCommands::Get().DistributeNodesHorizontally);
+						SubMenuSection.AddMenuEntry(FGraphEditorCommands::Get().DistributeNodesVertically);
+					}
+				}));
 		}
 
 		{
diff --git a/Engine/Source/Editor/UnrealEd/Private/MaterialGraphNode_Comment.cpp b/Engine/Source/Editor/UnrealEd/Private/MaterialGraphNode_Comment.cpp
--- a/Engine/Source/Editor/UnrealEd/Private/MaterialGraphNode_Comment.cpp
+++ b/Engine/Source/Editor/UnrealEd/Private/MaterialGraphNode_Comment.cpp
@@ -5,7 +5,7 @@
 =============================================================================*/
 
 #include "MaterialGraph/MaterialGraphNode_Comment.h"
-#include "Framework/MultiBox/MultiBoxBuilder.h"
+#include "ToolMenus.h"
 #include "MaterialGraph/MaterialGraphSchema.h"
 #include "Materials/MaterialExpressionComment.h"
 #include "Framework/Commands/GenericCommands.h"
@@ -67,19 +67,18 @@ void UMaterialGraphNode_Comment::PrepareForCopying()
 	}
 }
 
-void UMaterialGraphNode_Comment::GetContextMenuActions(const FGraphNodeContextMenuBuilder& Context) const
+void UMaterialGraphNode_Comment::GetNodeContextMenuActions(UToolMenu* Menu, UGraphNodeContextMenuContext* Context) const
 {
-	if (Context.Node && !Context.Pin)
+	if (Context->Node && !Context->Pin)
 	{
 		// frequently used common options
-		Context.MenuBuilder->BeginSection("MaterialEditorCommentMenu");
 		{
-			Context.MenuBuilder->AddMenuEntry( FGenericCommands::Get().Delete );
-			Context.MenuBuilder->AddMenuEntry( FGenericCommands::Get().Cut );
-			Context.MenuBuilder->AddMenuEntry( FGenericCommands::Get().Copy );
-			Context.MenuBuilder->AddMenuEntry( FGenericCommands::Get().Duplicate );
+			FToolMenuSection& Section = Menu->AddSection("MaterialEditorCommentMenu");
+			Section.AddMenuEntry(FGenericCommands::Get().Delete);
+			Section.AddMenuEntry(FGenericCommands::Get().Cut);
+			Section.AddMenuEntry(FGenericCommands::Get().Copy);
+			Section.AddMenuEntry(FGenericCommands::Get().Duplicate);
 		}
-		Context.MenuBuilder->EndSection();
 	}
 }
 
diff --git a/Engine/Source/Editor/UnrealEd/Public/GraphEditor.h b/Engine/Source/Editor/UnrealEd/Public/GraphEditor.h
--- a/Engine/Source/Editor/UnrealEd/Public/GraphEditor.h
+++ b/Engine/Source/Editor/UnrealEd/Public/GraphEditor.h
@@ -18,6 +18,7 @@
 class UEdGraph;
 struct FNotificationInfo;
 struct Rect;
+class FMenuBuilder;
 
 DECLARE_DELEGATE_ThreeParams( FOnNodeTextCommitted, const FText&, ETextCommit::Type, UEdGraphNode* );
 DECLARE_DELEGATE_RetVal_ThreeParams( bool, FOnNodeVerifyTextCommit, const FText&, UEdGraphNode*, FText& );
diff --git a/Engine/Source/Runtime/Engine/Classes/EdGraph/EdGraphNode.h b/Engine/Source/Runtime/Engine/Classes/EdGraph/EdGraphNode.h
--- a/Engine/Source/Runtime/Engine/Classes/EdGraph/EdGraphNode.h
+++ b/Engine/Source/Runtime/Engine/Classes/EdGraph/EdGraphNode.h
@@ -151,25 +151,7 @@ private:
 	FNodeMetadata() {}
 };
 
-/** This is the context for a GetContextMenuActions call into a specific node. */
-struct ENGINE_API FGraphNodeContextMenuBuilder
-{
-	/** The blueprint associated with this context; may be NULL for non-Kismet related graphs. */
-	const UBlueprint* Blueprint;
-	/** The graph associated with this context. */
-	const UEdGraph* Graph;
-	/** The node associated with this context. */
-	const UEdGraphNode* Node;
-	/** The pin associated with this context; may be NULL when over a node. */
-	const UEdGraphPin* Pin;
-	/** The menu builder to append actions to. */
-	class FMenuBuilder* MenuBuilder;
-	/** Whether the graph editor is currently part of a debugging session (any non-debugging commands should be disabled). */
-	bool bIsDebugging;
-
-	FGraphNodeContextMenuBuilder(const UEdGraph* InGraph, const UEdGraphNode* InNode, const UEdGraphPin* InPin, class FMenuBuilder* InMenuBuilder, bool bInDebuggingMode);
-};
-
+/** This is the context for GetContextMenuActions and GetNodeContextMenuActions calls. */
 UCLASS()
 class ENGINE_API UGraphNodeContextMenuContext : public UObject
 {
@@ -834,12 +816,12 @@ public:
 	/** Create a new unique Guid for this node */
 	void CreateNewGuid();
 
-	/** Gets a list of actions that can be done to this particular node */
-	virtual void GetContextMenuActions(const FGraphNodeContextMenuBuilder& Context) const {}
-	
 	/** Gets a list of actions that can be done to this particular node */
 	virtual void GetNodeContextMenuActions(class UToolMenu* Menu, class UGraphNodeContextMenuContext* Context) const {}
 
+	/** Does the node context menu inherit parent class's menu */
+	virtual bool IncludeParentNodeContextMenu() const { return false; }
+
 	// Gives each visual node a chance to do final validation before it's node is harvested for use at runtime
 	virtual void ValidateNodeDuringCompilation(class FCompilerResultsLog& MessageLog) const {}
 
diff --git a/Engine/Source/Runtime/Engine/Private/EdGraph/EdGraphNode.cpp b/Engine/Source/Runtime/Engine/Private/EdGraph/EdGraphNode.cpp
--- a/Engine/Source/Runtime/Engine/Private/EdGraph/EdGraphNode.cpp
+++ b/Engine/Source/Runtime/Engine/Private/EdGraph/EdGraphNode.cpp
@@ -95,27 +95,6 @@ UEdGraphNode::FCreatePinParams::FCreatePinParams(const FEdGraphPinType& PinType)
 }
 #endif // WITH_EDITOR
 
-/////////////////////////////////////////////////////
-// FGraphNodeContextMenuBuilder
-
-FGraphNodeContextMenuBuilder::FGraphNodeContextMenuBuilder(const UEdGraph* InGraph, const UEdGraphNode* InNode, const UEdGraphPin* InPin, FMenuBuilder* InMenuBuilder, bool bInDebuggingMode)
-	: Blueprint(nullptr)
-	, Graph(InGraph)
-	, Node(InNode)
-	, Pin(InPin)
-	, MenuBuilder(InMenuBuilder)
-	, bIsDebugging(bInDebuggingMode)
-{
-#if WITH_EDITOR
-	Blueprint = FBlueprintEditorUtils::FindBlueprintForGraph(Graph);
-#endif
-
-	if (Pin)
-	{
-		Node = Pin->GetOwningNode();
-	}
-}
-
 /////////////////////////////////////////////////////
 // UGraphNodeContextMenuContext
 
diff --git a/Engine/Source/Runtime/Engine/Public/EngineFwd.h b/Engine/Source/Runtime/Engine/Public/EngineFwd.h
--- a/Engine/Source/Runtime/Engine/Public/EngineFwd.h
+++ b/Engine/Source/Runtime/Engine/Public/EngineFwd.h
@@ -201,7 +201,6 @@ struct FExpressionInput;
 struct FGraphActionListBuilderBase;
 struct FGraphContextMenuBuilder;
 struct FGraphDisplayInfo;
-struct FGraphNodeContextMenuBuilder;
 struct FHitResult;
 struct FIntegralCurve;
 struct FKAggregateGeom;
