Fix sound class editor missing an include
#rnx

diff --git a/Engine/Source/Editor/AudioEditor/Private/SoundClassEditor.cpp b/Engine/Source/Editor/AudioEditor/Private/SoundClassEditor.cpp
--- a/Engine/Source/Editor/AudioEditor/Private/SoundClassEditor.cpp
+++ b/Engine/Source/Editor/AudioEditor/Private/SoundClassEditor.cpp
@@ -30,6 +30,7 @@
 #include "SoundClassEditorCommands.h"
 #include "Audio/AudioDebug.h"
 #include "AudioDeviceManager.h"
+#include "Framework/MultiBox/MultiBoxBuilder.h"
 
 #define LOCTEXT_NAMESPACE "SoundClassEditor"
 DEFINE_LOG_CATEGORY_STATIC( LogSoundClassEditor, Log, All );
