Fix for bad submix state when you click "Break Links..." on the output pin of an input submix.

diff --git a/Engine/Source/Editor/AudioEditor/Private/SoundSubmixGraphSchema.cpp b/Engine/Source/Editor/AudioEditor/Private/SoundSubmixGraphSchema.cpp
--- a/Engine/Source/Editor/AudioEditor/Private/SoundSubmixGraphSchema.cpp
+++ b/Engine/Source/Editor/AudioEditor/Private/SoundSubmixGraphSchema.cpp
@@ -416,17 +416,37 @@ void USoundSubmixGraphSchema::BreakPinLinks(UEdGraphPin& TargetPin, bool bSendsN
 	{
 		if (USoundSubmixGraphNode* GraphNode = Cast<USoundSubmixGraphNode>(TargetPin.GetOwningNode()))
 		{
-			// Iterate through all child submixes
-			USoundSubmixBase* OutputSubmix = GraphNode->SoundSubmix;
-
-			// Note: If we ever support multiple parents for submixes, this will need to be modified.
-			for (USoundSubmixBase* InputSubmix : OutputSubmix->ChildSubmixes)
+			// If TargetPin is an input, We should break links to all child submixes of the submix that owns this pin.
+			if (TargetPin.Direction == EEdGraphPinDirection::EGPD_Input)
 			{
-				if (USoundSubmixWithParentBase* SubmixWithParent = Cast<USoundSubmixWithParentBase>(InputSubmix))
+				// Iterate through all child submixes
+				USoundSubmixBase* OutputSubmix = GraphNode->SoundSubmix;
+
+				// Note: If we ever support multiple parents for submixes, this will need to be modified.
+				for (USoundSubmixBase* InputSubmix : OutputSubmix->ChildSubmixes)
 				{
-					SubmixWithParent->ParentSubmix = nullptr;
-					SubmixWithParent->PostEditChange();
+					if (USoundSubmixWithParentBase* SubmixWithParent = Cast<USoundSubmixWithParentBase>(InputSubmix))
+					{
+						SubmixWithParent->ParentSubmix = nullptr;
+						SubmixWithParent->PostEditChange();
+					}
 				}
+
+				OutputSubmix->ChildSubmixes.Reset();
+				OutputSubmix->PostEditChange();
+			}
+			else if (TargetPin.Direction == EEdGraphPinDirection::EGPD_Output)
+			{
+				// If this is an output pin, break the connection between this submix and it's parent.
+				USoundSubmixWithParentBase* InputSubmix = CastChecked<USoundSubmixWithParentBase>(GraphNode->SoundSubmix);
+				USoundSubmixBase* OutputSubmix = InputSubmix->ParentSubmix;
+				check(OutputSubmix);
+
+				OutputSubmix->ChildSubmixes.Remove(InputSubmix);
+				InputSubmix->ParentSubmix = nullptr;
+
+				OutputSubmix->PostEditChange();
+				InputSubmix->PostEditChange();
 			}
 		}
 
