Minifix: remove ensure since it will fire when the wire isn't connected to a pin yet.

diff --git a/Engine/Source/Editor/AudioEditor/Private/SoundSubmixGraphSchema.cpp b/Engine/Source/Editor/AudioEditor/Private/SoundSubmixGraphSchema.cpp
--- a/Engine/Source/Editor/AudioEditor/Private/SoundSubmixGraphSchema.cpp
+++ b/Engine/Source/Editor/AudioEditor/Private/SoundSubmixGraphSchema.cpp
@@ -50,7 +50,7 @@ FSoundSubmixGraphConnectionDrawingPolicy::FSoundSubmixGraphConnectionDrawingPoli
 // Give specific editor modes a chance to highlight this connection or darken non-interesting connections
 void FSoundSubmixGraphConnectionDrawingPolicy::DetermineWiringStyle(UEdGraphPin* OutputPin, UEdGraphPin* InputPin, FConnectionParams& OutParams)
 {
-	if (!ensure(OutputPin && InputPin && GraphObj))
+	if (!(OutputPin && InputPin && GraphObj))
 	{
 		return;
 	}
