Fix for import from template.
USoundWave::InvalidateCompressedData kicks off a recache of compressed data immediately by default. Issue is if you call USoundWave::InvalidateCompressedData on an empty USoundWave that just got created with NewObject, it would immediately try to compress RawPCMData, which hasn't been populated yet, and set NumChannels to 0 on the USoundWave.

diff --git a/Engine/Source/Editor/AudioEditor/Private/Factories/SoundFactory.cpp b/Engine/Source/Editor/AudioEditor/Private/Factories/SoundFactory.cpp
--- a/Engine/Source/Editor/AudioEditor/Private/Factories/SoundFactory.cpp
+++ b/Engine/Source/Editor/AudioEditor/Private/Factories/SoundFactory.cpp
@@ -493,10 +493,11 @@ UObject* USoundFactory::CreateObject
 		Sound->AssetImportData->Update(CurrentFilename);
 
 		// Compressed data is now out of date.
-		Sound->InvalidateCompressedData(true /* bFreeResources */);
+		const bool bRebuildStreamingChunks = FPlatformCompressionUtilities::IsCurrentPlatformUsingStreamCaching();
+		Sound->InvalidateCompressedData(true /* bFreeResources */, bRebuildStreamingChunks);
 
 		// If stream caching is enabled, we need to make sure this asset is ready for playback.
-		if (FPlatformCompressionUtilities::IsCurrentPlatformUsingStreamCaching() && Sound->IsStreaming(nullptr))
+		if (bRebuildStreamingChunks && Sound->IsStreaming(nullptr))
 		{
 			Sound->EnsureZerothChunkIsLoaded();
 		}
