Fix for submix editor leaving submixes in an invalid state.

diff --git a/Engine/Source/Editor/AudioEditor/Private/SoundSubmixGraph.cpp b/Engine/Source/Editor/AudioEditor/Private/SoundSubmixGraph.cpp
--- a/Engine/Source/Editor/AudioEditor/Private/SoundSubmixGraph.cpp
+++ b/Engine/Source/Editor/AudioEditor/Private/SoundSubmixGraph.cpp
@@ -91,7 +91,7 @@ void USoundSubmixGraph::LinkSoundSubmixes()
 					ChildSubmix->Modify();
 					if (USoundSubmixWithParentBase* SubmixWithParent = CastChecked<USoundSubmixWithParentBase>(ChildSubmix))
 					{
-						SubmixWithParent = nullptr;
+						SubmixWithParent->ParentSubmix = nullptr;
 					}
 				}
 			}
