Fix missing CanFilter forward in SoundEffectPreset

diff --git a/Engine/Source/Editor/AudioEditor/Private/AssetTypeActions/AssetTypeActions_SoundEffectPreset.h b/Engine/Source/Editor/AudioEditor/Private/AssetTypeActions/AssetTypeActions_SoundEffectPreset.h
--- a/Engine/Source/Editor/AudioEditor/Private/AssetTypeActions/AssetTypeActions_SoundEffectPreset.h
+++ b/Engine/Source/Editor/AudioEditor/Private/AssetTypeActions/AssetTypeActions_SoundEffectPreset.h
@@ -53,6 +53,7 @@ public:
 	FAssetTypeActions_SoundEffectPreset(USoundEffectPreset* InEffectPreset);
 
 	//~ Begin FAssetTypeActions_Base
+	virtual bool CanFilter() override { return EffectPreset->CanFilter(); }
 	virtual FText GetName() const override;
 	virtual FColor GetTypeColor() const override { return EffectPreset->GetPresetColor(); }
 	virtual const TArray<FText>& GetSubMenus() const override;
