Add Modulation Parameter Layout that disables additional functionality when plugin is disabled
#rnx

diff --git a/Engine/Source/Editor/AudioEditor/Private/AudioEditorModule.cpp b/Engine/Source/Editor/AudioEditor/Private/AudioEditorModule.cpp
--- a/Engine/Source/Editor/AudioEditor/Private/AudioEditorModule.cpp
+++ b/Engine/Source/Editor/AudioEditor/Private/AudioEditorModule.cpp
@@ -12,13 +12,14 @@
 #include "AssetToolsModule.h"
 #include "PropertyEditorModule.h"
 #include "SoundClassEditor.h"
+#include "Sound/AudioSettings.h"
 #include "Sound/SoundCue.h"
 #include "Sound/SoundWave.h"
 #include "Sound/SoundSubmix.h"
 #include "Sound/SoundEffectPreset.h"
 #include "SoundCueEditor.h"
+#include "SoundModulationParameterLayout.h"
 #include "SoundSubmixEditor.h"
-#include "Sound/AudioSettings.h"
 #include "AssetTypeActions/AssetTypeActions_DialogueVoice.h"
 #include "AssetTypeActions/AssetTypeActions_DialogueWave.h"
 #include "AssetTypeActions/AssetTypeActions_SoundAttenuation.h"
@@ -92,6 +93,10 @@ public:
 		// Create reimport handler for sound node waves
 		UReimportSoundFactory::StaticClass();
 
+		FPropertyEditorModule& PropertyModule = FModuleManager::LoadModuleChecked<FPropertyEditorModule>(TEXT("PropertyEditor"));
+		PropertyModule.RegisterCustomPropertyTypeLayout("SoundModulationParameterSettings",
+			FOnGetPropertyTypeCustomizationInstance::CreateStatic(&FSoundModulationParameterLayoutCustomization::MakeInstance));
+
 		SetupIcons();
 #if WITH_SNDFILE_IO
 		if (!Audio::InitSoundFileIOManager())
diff --git a/Engine/Source/Editor/AudioEditor/Private/SoundModulationParameterLayout.cpp b/Engine/Source/Editor/AudioEditor/Private/SoundModulationParameterLayout.cpp
new file mode 100644
--- /dev/null
+++ b/Engine/Source/Editor/AudioEditor/Private/SoundModulationParameterLayout.cpp
@@ -0,0 +1,175 @@
+// Copyright Epic Games, Inc. All Rights Reserved.
+#include "SoundModulationParameterLayout.h"
+
+#include "AudioDevice.h"
+#include "AudioDeviceManager.h"
+#include "DetailCategoryBuilder.h"
+#include "DetailLayoutBuilder.h"
+#include "DetailWidgetRow.h"
+#include "EditorStyleSet.h"
+#include "IAudioModulation.h"
+#include "Sound/SoundModulationParameter.h"
+#include "Widgets/Input/SButton.h"
+#include "Widgets/Text/STextBlock.h"
+#include "Widgets/SBoxPanel.h"
+#include "Widgets/DeclarativeSyntaxSupport.h"
+
+
+#define LOCTEXT_NAMESPACE "SoundModulationParameter"
+
+namespace ModParamLayoutUtils
+{
+	bool IsModulationEnabled()
+	{
+		if (FAudioDeviceManager* DeviceManager = FAudioDeviceManager::Get())
+		{
+			TArray<FAudioDevice*> AudioDevices = DeviceManager->GetAudioDevices();
+			for (FAudioDevice* AudioDevice : AudioDevices)
+			{
+				if (AudioDevice && AudioDevice->ModulationInterface.IsValid() && AudioDevice->IsModulationPluginEnabled())
+				{
+					return true;
+				}
+			}
+		}
+
+		return false;
+	}
+
+	void SetMetaData(FName FieldName, const FString& InDefault, const TSharedRef<IPropertyHandle>& InHandle, TSharedRef<IPropertyHandle>& OutHandle)
+	{
+		if (InHandle->HasMetaData(FieldName))
+		{
+			const FString& Value = InHandle->GetMetaData(FieldName);
+			OutHandle->SetInstanceMetaData(FieldName, Value);
+		}
+		else
+		{
+			OutHandle->SetInstanceMetaData(FieldName, InDefault);
+		}
+	}
+} // namespace ModParamLayoutUtils
+
+void FSoundModulationParameterLayoutCustomization::CustomizeHeader(TSharedRef<IPropertyHandle> StructPropertyHandle, FDetailWidgetRow& HeaderRow, IPropertyTypeCustomizationUtils& CustomizationUtils)
+{
+}
+
+void FSoundModulationParameterLayoutCustomization::CustomizeChildren(TSharedRef<IPropertyHandle> StructPropertyHandle, IDetailChildrenBuilder& ChildBuilder, IPropertyTypeCustomizationUtils& StructCustomizationUtils)
+{
+	TMap<FName, TSharedPtr<IPropertyHandle>> PropertyHandles;
+	uint32 NumChildren;
+	StructPropertyHandle->GetNumChildren(NumChildren);
+
+	for (uint32 ChildIndex = 0; ChildIndex < NumChildren; ++ChildIndex)
+	{
+		TSharedRef<IPropertyHandle> ChildHandle = StructPropertyHandle->GetChildHandle(ChildIndex).ToSharedRef();
+		const FName PropertyName = ChildHandle->GetProperty()->GetFName();
+		PropertyHandles.Add(PropertyName, ChildHandle);
+	}
+
+	TSharedRef<IPropertyHandle>ValueHandle = PropertyHandles.FindChecked(GET_MEMBER_NAME_CHECKED(FSoundModulationParameterSettings, Value)).ToSharedRef();
+
+	ModParamLayoutUtils::SetMetaData("ClampMin", TEXT("0.0"), StructPropertyHandle, ValueHandle);
+	ModParamLayoutUtils::SetMetaData("ClampMax", TEXT("1.0"), StructPropertyHandle, ValueHandle);
+	ModParamLayoutUtils::SetMetaData("UIMin",	 TEXT("0.0"), StructPropertyHandle, ValueHandle);
+	ModParamLayoutUtils::SetMetaData("UIMax",	 TEXT("1.0"), StructPropertyHandle, ValueHandle);
+
+	const FText NoModDisplayName = StructPropertyHandle->GetPropertyDisplayName();
+	FDetailWidgetRow& ValueNoModRow = ChildBuilder.AddCustomRow(NoModDisplayName);
+	ValueNoModRow.NameContent()
+		[
+			SNew(STextBlock)
+			.Font(IDetailLayoutBuilder::GetDetailFont())
+			.Text(NoModDisplayName)
+			.ToolTipText(StructPropertyHandle->GetToolTipText())
+		]
+		.ValueContent()
+		.MinDesiredWidth(120.0f)
+		[
+			SNew(SHorizontalBox)
+			+ SHorizontalBox::Slot()
+				.FillWidth(0.5f)
+				.Padding(1.0f, 0.0f, 0.0f, 0.0f)
+				.VAlign(VAlign_Center)
+				[
+					ValueHandle->CreatePropertyValueWidget()
+				]
+		];
+	ValueNoModRow.Visibility(TAttribute<EVisibility>::Create([this]()
+	{
+		return ModParamLayoutUtils::IsModulationEnabled() ? EVisibility::Hidden : EVisibility::Visible;
+	}));
+
+	TSharedRef<IPropertyHandle>OperatorHandle = PropertyHandles.FindChecked(GET_MEMBER_NAME_CHECKED(FSoundModulationParameterSettings, Operator)).ToSharedRef();
+
+	static const FText OpToolTipText = LOCTEXT("SoundModulationOperatorToolTip", "");
+
+	const FText DisplayName = StructPropertyHandle->GetPropertyDisplayName();
+	FDetailWidgetRow& ValueRow = ChildBuilder.AddCustomRow(DisplayName);
+	ValueRow.NameContent()
+		[
+			SNew(STextBlock)
+			.Font(IDetailLayoutBuilder::GetDetailFont())
+			.Text(DisplayName)
+			.ToolTipText(StructPropertyHandle->GetToolTipText())
+		]
+		.ValueContent()
+		.MinDesiredWidth(200.0f)
+		[
+			SNew(SHorizontalBox)
+			+ SHorizontalBox::Slot()
+				.FillWidth(0.3f)
+				.Padding(1.0f, 0.0f, 0.0f, 0.0f)
+				.VAlign(VAlign_Center)
+				[
+					ValueHandle->CreatePropertyValueWidget()
+				]
+			+ SHorizontalBox::Slot()
+				.FillWidth(0.7f)
+				.Padding(1.0f, 0.0f, 0.0f, 0.0f)
+				.VAlign(VAlign_Center)
+				[
+					OperatorHandle->CreatePropertyValueWidget()
+				]
+		];
+		ValueRow.Visibility(TAttribute<EVisibility>::Create([this]()
+		{
+			return ModParamLayoutUtils::IsModulationEnabled() ? EVisibility::Visible : EVisibility::Hidden;
+		}));
+
+		TSharedRef<IPropertyHandle>ModHandle = PropertyHandles.FindChecked(GET_MEMBER_NAME_CHECKED(FSoundModulationParameterSettings, Modulator)).ToSharedRef();
+		FDetailWidgetRow& ModulatorRow = ChildBuilder.AddCustomRow(DisplayName)
+		.NameContent()
+		[
+			SNew(STextBlock)
+			.Font(IDetailLayoutBuilder::GetDetailFont())
+			.Text(FText::Format(LOCTEXT("SoundModulationParameterModulator", "{0} Modulator"), DisplayName))
+			.ToolTipText(ModHandle->GetToolTipText())
+		]
+		.ValueContent()
+		.MinDesiredWidth(200.0f)
+		[
+			SNew(SHorizontalBox)
+			+ SHorizontalBox::Slot()
+			.FillWidth(1.0f)
+			.Padding(4.0f, 0.0f, 0.0f, 0.0f)
+			.VAlign(VAlign_Center)
+			[
+				ModHandle->CreatePropertyValueWidget()
+			]
+		];
+		ModulatorRow.Visibility(TAttribute<EVisibility>::Create([this, OperatorHandle]()
+		{
+			if (!ModParamLayoutUtils::IsModulationEnabled())
+			{
+				return EVisibility::Hidden;
+			}
+
+			FString EnumString;
+			OperatorHandle->GetValueAsDisplayString(EnumString);
+
+			static_assert(static_cast<uint8>(ESoundModulatorOperator::None) == 0, "None value reassigned. Ensure enum string is valid to compare against below.");
+			return EnumString == TEXT("None") ? EVisibility::Hidden : EVisibility::Visible;
+		}));
+}
+#undef LOCTEXT_NAMESPACE
diff --git a/Engine/Source/Editor/AudioEditor/Private/SoundModulationParameterLayout.h b/Engine/Source/Editor/AudioEditor/Private/SoundModulationParameterLayout.h
new file mode 100644
--- /dev/null
+++ b/Engine/Source/Editor/AudioEditor/Private/SoundModulationParameterLayout.h
@@ -0,0 +1,30 @@
+// Copyright Epic Games, Inc. All Rights Reserved.
+#pragma once
+
+#include "CoreMinimal.h"
+
+#include "IDetailChildrenBuilder.h"
+#include "IDetailPropertyRow.h"
+#include "IPropertyTypeCustomization.h"
+#include "Layout/Visibility.h"
+#include "Misc/Attribute.h"
+#include "PropertyHandle.h"
+
+
+// Forward Declarations
+struct FSoundModulationParameterSettings;
+
+
+class FSoundModulationParameterLayoutCustomization : public IPropertyTypeCustomization
+{
+public:
+	static TSharedRef<IPropertyTypeCustomization> MakeInstance()
+	{
+		return MakeShared<FSoundModulationParameterLayoutCustomization>();
+	}
+
+	//~ Begin IPropertyTypeCustomization
+	virtual void CustomizeHeader(TSharedRef<IPropertyHandle> StructPropertyHandle, FDetailWidgetRow& HeaderRow, IPropertyTypeCustomizationUtils& StructCustomizationUtils) override;
+	virtual void CustomizeChildren(TSharedRef<IPropertyHandle> StructPropertyHandle, IDetailChildrenBuilder& ChildBuilder, IPropertyTypeCustomizationUtils& StructCustomizationUtils) override;
+	//~ End IPropertyTypeCustomization
+};
