Remove deprecated asset editor manager.

diff --git a/Engine/Plugins/2D/Paper2D/Source/Paper2DEditor/Paper2DEditor.Build.cs b/Engine/Plugins/2D/Paper2D/Source/Paper2DEditor/Paper2DEditor.Build.cs
--- a/Engine/Plugins/2D/Paper2D/Source/Paper2DEditor/Paper2DEditor.Build.cs
+++ b/Engine/Plugins/2D/Paper2D/Source/Paper2DEditor/Paper2DEditor.Build.cs
@@ -17,7 +17,7 @@ public class Paper2DEditor : ModuleRules
 				"SlateCore",
 				"Engine",
 				"InputCore",
-				"UnrealEd", // for FAssetEditorManager
+				"UnrealEd", // for AssetEditorSubsystem
 				"KismetWidgets",
 				"Kismet",  // for FWorkflowCentricApplication
 				"PropertyEditor",
diff --git a/Engine/Plugins/Editor/AssetSearch/Source/Private/Widgets/SSearchTreeRow.h b/Engine/Plugins/Editor/AssetSearch/Source/Private/Widgets/SSearchTreeRow.h
--- a/Engine/Plugins/Editor/AssetSearch/Source/Private/Widgets/SSearchTreeRow.h
+++ b/Engine/Plugins/Editor/AssetSearch/Source/Private/Widgets/SSearchTreeRow.h
@@ -3,7 +3,6 @@
 #pragma once
 
 #include "ClassIconFinder.h"
-#include "Toolkits/AssetEditorManager.h"
 #include "CoreMinimal.h"
 #include "Widgets/SWidget.h"
 #include "Widgets/Views/STableRow.h"
diff --git a/Engine/Plugins/Enterprise/DataprepEditor/Source/DataprepEditor/Private/DataprepEditor.cpp b/Engine/Plugins/Enterprise/DataprepEditor/Source/DataprepEditor/Private/DataprepEditor.cpp
--- a/Engine/Plugins/Enterprise/DataprepEditor/Source/DataprepEditor/Private/DataprepEditor.cpp
+++ b/Engine/Plugins/Enterprise/DataprepEditor/Source/DataprepEditor/Private/DataprepEditor.cpp
@@ -633,7 +633,7 @@ void FDataprepEditor::OnExecutePipeline()
 	{
 		FTimeLogger TimeLogger( TEXT("ExecutePipeline") );
 
-		// Some operation can indirectly call FAssetEditorManager::CloseAllAssetEditors (eg. Remove Asset)
+		// Some operation can indirectly call UAssetEditorSubsystem::CloseAllAssetEditors (eg. Remove Asset)
 		// Editors can individually refuse this request: we ignore it during the pipeline traversal.
 		TGuardValue<bool> IgnoreCloseRequestGuard(bIgnoreCloseRequest, true);
 
diff --git a/Engine/Plugins/Enterprise/VariantManagerContent/Source/VariantManagerContentEditor/Private/SwitchActorCustomization.cpp b/Engine/Plugins/Enterprise/VariantManagerContent/Source/VariantManagerContentEditor/Private/SwitchActorCustomization.cpp
--- a/Engine/Plugins/Enterprise/VariantManagerContent/Source/VariantManagerContentEditor/Private/SwitchActorCustomization.cpp
+++ b/Engine/Plugins/Enterprise/VariantManagerContent/Source/VariantManagerContentEditor/Private/SwitchActorCustomization.cpp
@@ -14,7 +14,6 @@
 #include "IAssetTools.h"
 #include "Modules/ModuleManager.h"
 #include "ScopedTransaction.h"
-#include "Toolkits/AssetEditorManager.h"
 #include "UObject/Object.h"
 #include "Widgets/DeclarativeSyntaxSupport.h"
 #include "Widgets/Input/SComboBox.h"
diff --git a/Engine/Plugins/Experimental/ControlRig/Source/ControlRigEditor/Private/Sequencer/ControlRigParameterTrackEditor.cpp b/Engine/Plugins/Experimental/ControlRig/Source/ControlRigEditor/Private/Sequencer/ControlRigParameterTrackEditor.cpp
--- a/Engine/Plugins/Experimental/ControlRig/Source/ControlRigEditor/Private/Sequencer/ControlRigParameterTrackEditor.cpp
+++ b/Engine/Plugins/Experimental/ControlRig/Source/ControlRigEditor/Private/Sequencer/ControlRigParameterTrackEditor.cpp
@@ -35,7 +35,6 @@
 #include "Misc/MessageDialog.h"
 #include "Framework/Notifications/NotificationManager.h"
 #include "Widgets/Notifications/SNotificationList.h"
-#include "Toolkits/AssetEditorManager.h"
 #include "Engine/SCS_Node.h"
 #include "Engine/SimpleConstructionScript.h"
 #include "Engine/Blueprint.h"
diff --git a/Engine/Plugins/Runtime/Analytics/AnalyticsMulticast/Source/AnaltyicsMulticastEditor/AnalyticsMulticastEditor.Build.cs b/Engine/Plugins/Runtime/Analytics/AnalyticsMulticast/Source/AnaltyicsMulticastEditor/AnalyticsMulticastEditor.Build.cs
--- a/Engine/Plugins/Runtime/Analytics/AnalyticsMulticast/Source/AnaltyicsMulticastEditor/AnalyticsMulticastEditor.Build.cs
+++ b/Engine/Plugins/Runtime/Analytics/AnalyticsMulticast/Source/AnaltyicsMulticastEditor/AnalyticsMulticastEditor.Build.cs
@@ -15,7 +15,7 @@ public class AnalyticsMulticastEditor : ModuleRules
 				"Slate",
 				"SlateCore",
 				"Engine",
-				"UnrealEd", // for FAssetEditorManager
+				"UnrealEd", // for Asset Editor Subsystem
 				"PropertyEditor",
 				"WorkspaceMenuStructure",
 				"EditorStyle",
diff --git a/Engine/Plugins/VirtualProduction/Takes/Source/TakeRecorderSources/Private/TakeRecorderSourcesModule.cpp b/Engine/Plugins/VirtualProduction/Takes/Source/TakeRecorderSources/Private/TakeRecorderSourcesModule.cpp
--- a/Engine/Plugins/VirtualProduction/Takes/Source/TakeRecorderSources/Private/TakeRecorderSourcesModule.cpp
+++ b/Engine/Plugins/VirtualProduction/Takes/Source/TakeRecorderSources/Private/TakeRecorderSourcesModule.cpp
@@ -31,7 +31,7 @@
 #include "Editor.h"
 #include "SceneOutlinerModule.h"
 #include "SceneOutlinerPublicTypes.h"
-#include "Toolkits/AssetEditorManager.h"
+#include "Subsystems/AssetEditorSubsystem.h"
 
 #include "TakeRecorderMicrophoneAudioSource.h"
 #include "TakeRecorderWorldSource.h"
diff --git a/Engine/Source/Developer/AssetTools/Public/AssetTypeActions_Base.h b/Engine/Source/Developer/AssetTools/Public/AssetTypeActions_Base.h
--- a/Engine/Source/Developer/AssetTools/Public/AssetTypeActions_Base.h
+++ b/Engine/Source/Developer/AssetTools/Public/AssetTypeActions_Base.h
@@ -11,7 +11,6 @@
 #include "IAssetTools.h"
 #include "IAssetTypeActions.h"
 #include "AssetToolsModule.h"
-#include "Toolkits/AssetEditorManager.h"
 #include "Toolkits/SimpleAssetEditor.h"
 #include "ARFilter.h"
 
@@ -47,24 +46,6 @@ public:
 		FSimpleAssetEditor::CreateEditor(EToolkitMode::Standalone, EditWithinLevelEditor, InObjects);
 	}
 
-PRAGMA_DISABLE_DEPRECATION_WARNINGS
-UE_DEPRECATED(4.24, "Use AssetsActivatedOverride instead to provide any non-default behavior. Using AssetsActivatedOverride, you no longer need a call to FAssetTypeActions_Base::AssetsActivated.")
-	virtual void AssetsActivated( const TArray<UObject*>& InObjects, EAssetTypeActivationMethod::Type ActivationType ) override
-	{
-		if (ActivationType == EAssetTypeActivationMethod::DoubleClicked || ActivationType == EAssetTypeActivationMethod::Opened)
-		{
-			if (InObjects.Num() == 1)
-			{
-				FAssetEditorManager::Get().OpenEditorForAsset(InObjects[0]);
-			}
-			else if (InObjects.Num() > 1)
-			{
-				FAssetEditorManager::Get().OpenEditorForAssets(InObjects);
-			}
-		}
-	}
-PRAGMA_ENABLE_DEPRECATION_WARNINGS
-
 	virtual bool AssetsActivatedOverride(const TArray<UObject*>& InObjects, EAssetTypeActivationMethod::Type ActivationType) override
 	{
 		return false;
diff --git a/Engine/Source/Developer/AssetTools/Public/IAssetTypeActions.h b/Engine/Source/Developer/AssetTools/Public/IAssetTypeActions.h
--- a/Engine/Source/Developer/AssetTools/Public/IAssetTypeActions.h
+++ b/Engine/Source/Developer/AssetTools/Public/IAssetTypeActions.h
@@ -63,9 +63,6 @@ public:
 	/** Opens the asset editor for the specified objects. If EditWithinLevelEditor is valid, the world-centric editor will be used. */
 	virtual void OpenAssetEditor( const TArray<UObject*>& InObjects, TSharedPtr<IToolkitHost> EditWithinLevelEditor = TSharedPtr<IToolkitHost>() ) = 0;
 
-	/** Performs asset type specific activation for the supplied assets. This happens when the user double clicks, presses enter, or presses space. */
-	virtual void AssetsActivated( const TArray<UObject*>& InObjects, EAssetTypeActivationMethod::Type ActivationType ) = 0;
-
 	/** Allows overriding asset activation to perform asset type specific activation for the supplied assets. This happens when the user double clicks, presses enter, or presses space. Return true if you have overridden the behavior. */
 	virtual bool AssetsActivatedOverride(const TArray<UObject*>& InObjects, EAssetTypeActivationMethod::Type ActivationType) = 0;
 
diff --git a/Engine/Source/Developer/HierarchicalLODUtilities/Private/HierarchicalLODUtilities.cpp b/Engine/Source/Developer/HierarchicalLODUtilities/Private/HierarchicalLODUtilities.cpp
--- a/Engine/Source/Developer/HierarchicalLODUtilities/Private/HierarchicalLODUtilities.cpp
+++ b/Engine/Source/Developer/HierarchicalLODUtilities/Private/HierarchicalLODUtilities.cpp
@@ -31,7 +31,6 @@
 
 #if WITH_EDITOR
 #include "Editor.h"
-#include "Toolkits/AssetEditorManager.h"
 #include "ScopedTransaction.h"
 #include "PackageTools.h"
 #include "Settings/EditorExperimentalSettings.h"
diff --git a/Engine/Source/Developer/Merge/Private/SBlueprintMerge.cpp b/Engine/Source/Developer/Merge/Private/SBlueprintMerge.cpp
--- a/Engine/Source/Developer/Merge/Private/SBlueprintMerge.cpp
+++ b/Engine/Source/Developer/Merge/Private/SBlueprintMerge.cpp
@@ -10,7 +10,6 @@
 #include "Engine/Blueprint.h"
 #include "Misc/PackageName.h"
 #include "Editor.h"
-#include "Toolkits/AssetEditorManager.h"
 #include "Misc/MessageDialog.h"
 #include "HAL/FileManager.h"
 #include "EditorStyleSet.h"
diff --git a/Engine/Source/Developer/MeshUtilities/Private/MeshUtilities.cpp b/Engine/Source/Developer/MeshUtilities/Private/MeshUtilities.cpp
--- a/Engine/Source/Developer/MeshUtilities/Private/MeshUtilities.cpp
+++ b/Engine/Source/Developer/MeshUtilities/Private/MeshUtilities.cpp
@@ -63,7 +63,6 @@
 #include "LandscapeHeightfieldCollisionComponent.h"
 #include "Engine/MeshMergeCullingVolume.h"
 
-#include "Toolkits/AssetEditorManager.h"
 #include "LevelEditor.h"
 #include "IAnimationBlueprintEditor.h"
 #include "IAnimationBlueprintEditorModule.h"
diff --git a/Engine/Source/Editor/AudioEditor/Private/AssetTypeActions/AssetTypeActions_SoundBase.cpp b/Engine/Source/Editor/AudioEditor/Private/AssetTypeActions/AssetTypeActions_SoundBase.cpp
--- a/Engine/Source/Editor/AudioEditor/Private/AssetTypeActions/AssetTypeActions_SoundBase.cpp
+++ b/Engine/Source/Editor/AudioEditor/Private/AssetTypeActions/AssetTypeActions_SoundBase.cpp
@@ -74,49 +74,6 @@ void FAssetTypeActions_SoundBase::GetActions(const TArray<UObject*>& InObjects,
 	);
 }
 
-PRAGMA_DISABLE_DEPRECATION_WARNINGS
-void FAssetTypeActions_SoundBase::AssetsActivated( const TArray<UObject*>& InObjects, EAssetTypeActivationMethod::Type ActivationType )
-{
-	if (ActivationType == EAssetTypeActivationMethod::Previewed)
-	{
-		USoundBase* TargetSound = NULL;
-
-		for (auto ObjIt = InObjects.CreateConstIterator(); ObjIt; ++ObjIt)
-		{
-			TargetSound = Cast<USoundBase>(*ObjIt);
-			if ( TargetSound )
-			{
-				// Only target the first valid sound cue
-				break;
-			}
-		}
-
-		UAudioComponent* PreviewComp = GEditor->GetPreviewAudioComponent();
-		if ( PreviewComp && PreviewComp->IsPlaying() )
-		{
-			// Already previewing a sound, if it is the target cue then stop it, otherwise play the new one
-			if ( !TargetSound || PreviewComp->Sound == TargetSound )
-			{
-				StopSound();
-			}
-			else
-			{
-				PlaySound(TargetSound);
-			}
-		}
-		else
-		{
-			// Not already playing, play the target sound cue if it exists
-			PlaySound(TargetSound);
-		}
-	}
-	else
-	{
-		FAssetTypeActions_Base::AssetsActivated(InObjects, ActivationType);
-	}
-}
-PRAGMA_ENABLE_DEPRECATION_WARNINGS
-
 bool FAssetTypeActions_SoundBase::AssetsActivatedOverride(const TArray<UObject*>& InObjects, EAssetTypeActivationMethod::Type ActivationType)
 {
 	if (ActivationType == EAssetTypeActivationMethod::Previewed)
diff --git a/Engine/Source/Editor/AudioEditor/Private/SoundSubmixGraphNode.cpp b/Engine/Source/Editor/AudioEditor/Private/SoundSubmixGraphNode.cpp
--- a/Engine/Source/Editor/AudioEditor/Private/SoundSubmixGraphNode.cpp
+++ b/Engine/Source/Editor/AudioEditor/Private/SoundSubmixGraphNode.cpp
@@ -4,9 +4,9 @@
 #include "Sound/SoundSubmix.h"
 #include "SoundSubmixGraph/SoundSubmixGraphSchema.h"
 #include "SoundSubmixGraph/SoundSubmixGraph.h"
-#include "Toolkits/AssetEditorManager.h"
 #include "SoundSubmixEditor.h"
 #include "SoundSubmixDefaultColorPalette.h"
+#include "Subsystems/AssetEditorSubsystem.h"
 
 #define LOCTEXT_NAMESPACE "SoundSubmixGraphNode"
 
diff --git a/Engine/Source/Editor/AudioEditor/Private/SoundSubmixGraphSchema.cpp b/Engine/Source/Editor/AudioEditor/Private/SoundSubmixGraphSchema.cpp
--- a/Engine/Source/Editor/AudioEditor/Private/SoundSubmixGraphSchema.cpp
+++ b/Engine/Source/Editor/AudioEditor/Private/SoundSubmixGraphSchema.cpp
@@ -17,7 +17,7 @@
 #include "SoundSubmixGraph/SoundSubmixGraph.h"
 #include "SoundSubmixEditor.h"
 #include "SoundSubmixEditorUtilities.h"
-#include "Toolkits/AssetEditorManager.h"
+#include "Subsystems/AssetEditorSubsystem.h"
 #include "ToolMenus.h"
 #include "SoundSubmixDefaultColorPalette.h"
 
diff --git a/Engine/Source/Editor/AudioEditor/Public/AssetTypeActions/AssetTypeActions_SoundBase.h b/Engine/Source/Editor/AudioEditor/Public/AssetTypeActions/AssetTypeActions_SoundBase.h
--- a/Engine/Source/Editor/AudioEditor/Public/AssetTypeActions/AssetTypeActions_SoundBase.h
+++ b/Engine/Source/Editor/AudioEditor/Public/AssetTypeActions/AssetTypeActions_SoundBase.h
@@ -18,7 +18,6 @@ public:
 	virtual UClass* GetSupportedClass() const override;
 	virtual bool HasActions ( const TArray<UObject*>& InObjects ) const override { return true; }
 	virtual void GetActions(const TArray<UObject*>& InObjects, struct FToolMenuSection& Section) override;
-	virtual void AssetsActivated( const TArray<UObject*>& InObjects, EAssetTypeActivationMethod::Type ActivationType ) override;
 	virtual bool AssetsActivatedOverride(const TArray<UObject*>& InObjects, EAssetTypeActivationMethod::Type ActivationType) override;
 	virtual uint32 GetCategories() override { return EAssetTypeCategories::Sounds; }
 	virtual bool CanFilter() override { return false; }
diff --git a/Engine/Source/Editor/DataTableEditor/Private/DataTableEditor.cpp b/Engine/Source/Editor/DataTableEditor/Private/DataTableEditor.cpp
--- a/Engine/Source/Editor/DataTableEditor/Private/DataTableEditor.cpp
+++ b/Engine/Source/Editor/DataTableEditor/Private/DataTableEditor.cpp
@@ -37,7 +37,6 @@
 #include "PropertyEditorModule.h"
 #include "UObject/StructOnScope.h"
 #include "Toolkits/GlobalEditorCommonCommands.h"
-#include "Toolkits/AssetEditorManager.h"
 #include "DataTableRowUtlis.h"
 #include "Engine/DataTable.h"
 #include "Subsystems/AssetEditorSubsystem.h"
diff --git a/Engine/Source/Editor/UMGEditor/UMGEditor.Build.cs b/Engine/Source/Editor/UMGEditor/UMGEditor.Build.cs
--- a/Engine/Source/Editor/UMGEditor/UMGEditor.Build.cs
+++ b/Engine/Source/Editor/UMGEditor/UMGEditor.Build.cs
@@ -42,7 +42,7 @@ public class UMGEditor : ModuleRules
 				"InputCore",
 				"Engine",
 				"AssetTools",
-				"UnrealEd", // for FAssetEditorManager
+				"UnrealEd", // for Asset Editor Subsystem
 				"KismetWidgets",
 				"KismetCompiler",
 				"BlueprintGraph",
diff --git a/Engine/Source/Editor/UnrealEd/Private/Subsystems/AssetEditorSubsystem.cpp b/Engine/Source/Editor/UnrealEd/Private/Subsystems/AssetEditorSubsystem.cpp
--- a/Engine/Source/Editor/UnrealEd/Private/Subsystems/AssetEditorSubsystem.cpp
+++ b/Engine/Source/Editor/UnrealEd/Private/Subsystems/AssetEditorSubsystem.cpp
@@ -729,7 +729,7 @@ void UAssetEditorSubsystem::RestorePreviouslyOpenAssets()
 
 void UAssetEditorSubsystem::SpawnRestorePreviouslyOpenAssetsNotification(const bool bCleanShutdown, const TArray<FString>& AssetsToOpen)
 {
-	/** Utility functions for the notification which don't rely on the state from FAssetEditorManager */
+	/** Utility functions for notifications */
 	struct Local
 	{
 		static ECheckBoxState GetDontAskAgainCheckBoxState()
diff --git a/Engine/Source/Editor/UnrealEd/Private/Toolkits/AssetEditorManager.cpp b/Engine/Source/Editor/UnrealEd/Private/Toolkits/AssetEditorManager.cpp
deleted file mode 100644
diff --git a/Engine/Source/Editor/UnrealEd/Private/Toolkits/AssetEditorToolkit.cpp b/Engine/Source/Editor/UnrealEd/Private/Toolkits/AssetEditorToolkit.cpp
--- a/Engine/Source/Editor/UnrealEd/Private/Toolkits/AssetEditorToolkit.cpp
+++ b/Engine/Source/Editor/UnrealEd/Private/Toolkits/AssetEditorToolkit.cpp
@@ -279,7 +279,7 @@ void FAssetEditorToolkit::InitAssetEditor( const EToolkitMode::Type Mode, const
 		EditorModeManager->SetToolkitHost(ToolkitHost.Pin().ToSharedRef());
 	}
 
-	// NOTE: Currently, the AssetEditorManager will keep a hard reference to our object as we're editing it
+	// NOTE: Currently, the AssetEditorSubsystem will keep a hard reference to our object as we're editing it
 	GEditor->GetEditorSubsystem<UAssetEditorSubsystem>()->NotifyAssetsOpened( EditingObjects, this );
 }
 
@@ -605,12 +605,12 @@ void FAssetEditorToolkit::SaveAssetAs_Execute()
 	{
 		if ((SavedObjects[Index] != ObjectsToSave[Index]) && (SavedObjects[Index] != nullptr))
 		{
-			AssetEditorManager.CloseAllEditorsForAsset(ObjectsToSave[Index]);
+			AssetEditorSubsystem->CloseAllEditorsForAsset(ObjectsToSave[Index]);
 		}
 	}
 
 	// reopen asset editor
-	AssetEditorManager.OpenEditorForAssets(TArrayBuilder<UObject*>().Add(SavedObjects[0]), ToolkitMode, MyToolkitHost.ToSharedRef());
+	AssetEditorSubsystem->OpenEditorForAssets(TArrayBuilder<UObject*>().Add(SavedObjects[0]), ToolkitMode, MyToolkitHost.ToSharedRef());
 	*/
 	// hack
 	TArray<UObject*> ObjectsToReopen;
diff --git a/Engine/Source/Editor/UnrealEd/Public/Toolkits/AssetEditorManager.h b/Engine/Source/Editor/UnrealEd/Public/Toolkits/AssetEditorManager.h
deleted file mode 100644
diff --git a/Engine/Source/Editor/UnrealEd/Public/Toolkits/AssetEditorToolkit.h b/Engine/Source/Editor/UnrealEd/Public/Toolkits/AssetEditorToolkit.h
--- a/Engine/Source/Editor/UnrealEd/Public/Toolkits/AssetEditorToolkit.h
+++ b/Engine/Source/Editor/UnrealEd/Public/Toolkits/AssetEditorToolkit.h
@@ -398,7 +398,7 @@ private:
 	static TSharedPtr<FExtensibilityManager> SharedToolBarExtensibilityManager;
 
 	/** The object we're currently editing */
-	// @todo toolkit minor: Currently we don't need to serialize this object reference because the FAssetEditorManager is kept in sync (and will always serialize it.)
+	// @todo toolkit minor: Currently we don't need to serialize this object reference because the AssetEditorSubsystem is kept in sync (and will always serialize it.)
 	TArray<UObject*> EditingObjects;
 	
 	/** Asset Editor Default Toolbar */
