Fix modulation destination UI/Clamp boundaries when modulation is not enabled

diff --git a/Engine/Source/Editor/AudioEditor/Private/SoundModulationDestinationLayout.cpp b/Engine/Source/Editor/AudioEditor/Private/SoundModulationDestinationLayout.cpp
--- a/Engine/Source/Editor/AudioEditor/Private/SoundModulationDestinationLayout.cpp
+++ b/Engine/Source/Editor/AudioEditor/Private/SoundModulationDestinationLayout.cpp
@@ -102,15 +102,8 @@ namespace ModDestinationLayoutUtils
 		return false;
 	}
 
-	void CustomizeChildren_AddValueRow(
-		IDetailChildrenBuilder& ChildBuilder,
-		TSharedRef<IPropertyHandle> StructPropertyHandle,
-		TSharedRef<IPropertyHandle> ValueHandle,
-		TSharedRef<IPropertyHandle> ModulatorHandle,
-		TSharedRef<IPropertyHandle> EnablementHandle)
+	FText SetMetaData(TSharedRef<IPropertyHandle> StructPropertyHandle, TSharedRef<IPropertyHandle> ValueHandle, FText& OutUnitDisplayText, FName& OutParamName)
 	{
-		FText UnitDisplayText;
-
 		bool bClampValuesSet = false;
 		float ClampMinValue = 0.0f;
 		float ClampMaxValue = 1.0f;
@@ -130,22 +123,22 @@ namespace ModDestinationLayoutUtils
 			bClampValuesSet = true;
 		}
 
-		const FName ParamName = ModDestinationLayoutUtils::GetParameterNameFromMetaData(StructPropertyHandle);
-		if (ParamName != FName())
+		OutParamName = ModDestinationLayoutUtils::GetParameterNameFromMetaData(StructPropertyHandle);
+		if (OutParamName != FName())
 		{
 			// If parameter was provided, it overrides ClampMin/Max.  User data however overrides UIMin/Max if its
 			// in clamp range.
 			if (IAudioModulation* ModulationInterface = ModDestinationLayoutUtils::GetEditorModulationInterface())
 			{
-				Audio::FModulationParameter Parameter = ModulationInterface->GetParameter(ParamName);
+				Audio::FModulationParameter Parameter = ModulationInterface->GetParameter(OutParamName);
 				UIMinValue = Parameter.MinValue;
 				UIMaxValue = Parameter.MaxValue;
 				ClampMinValue = UIMinValue;
 				ClampMaxValue = UIMaxValue;
-				UnitDisplayText = Parameter.UnitDisplayName;
+				OutUnitDisplayText = Parameter.UnitDisplayName;
 				if (bClampValuesSet)
 				{
-					UE_LOG(LogAudioEditor, Warning, TEXT("ClampMin/Max overridden by AudioModulation plugin asset with ParamName '%s'."), *ParamName.ToString());
+					UE_LOG(LogAudioEditor, Warning, TEXT("ClampMin/Max overridden by AudioModulation plugin asset with ParamName '%s'."), *OutParamName.ToString());
 				}
 			}
 
@@ -154,7 +147,7 @@ namespace ModDestinationLayoutUtils
 				float NewMin = UIMinValue;
 				FString ParamString = StructPropertyHandle->GetMetaData("UIMin");
 				NewMin = FCString::Atof(*ParamString);
-				UIMaxValue = FMath::Clamp(NewMin, ClampMinValue, ClampMaxValue);
+				UIMinValue = FMath::Clamp(NewMin, ClampMinValue, ClampMaxValue);
 			}
 
 			if (StructPropertyHandle->HasMetaData("UIMax"))
@@ -170,6 +163,20 @@ namespace ModDestinationLayoutUtils
 		ValueHandle->SetInstanceMetaData("ClampMax", FString::Printf(TEXT("%f"), ClampMaxValue));
 		ValueHandle->SetInstanceMetaData("UIMin", FString::Printf(TEXT("%f"), UIMinValue));
 		ValueHandle->SetInstanceMetaData("UIMax", FString::Printf(TEXT("%f"), UIMaxValue));
+		
+		return OutUnitDisplayText;
+	}
+
+	void CustomizeChildren_AddValueRow(
+		IDetailChildrenBuilder& ChildBuilder,
+		TSharedRef<IPropertyHandle> StructPropertyHandle,
+		TSharedRef<IPropertyHandle> ValueHandle,
+		TSharedRef<IPropertyHandle> ModulatorHandle,
+		TSharedRef<IPropertyHandle> EnablementHandle)
+	{
+		FText UnitDisplayText = FText::GetEmpty();
+		FName ParamName;
+		SetMetaData(StructPropertyHandle, ValueHandle, UnitDisplayText, ParamName);
 
 		const FText DisplayName = StructPropertyHandle->GetPropertyDisplayName();
 		ChildBuilder.AddCustomRow(DisplayName)
@@ -284,6 +291,10 @@ namespace ModDestinationLayoutUtils
 
 	void CustomizeChildren_AddValueNoModRow(IDetailChildrenBuilder& ChildBuilder, TSharedRef<IPropertyHandle> StructPropertyHandle, TSharedRef<IPropertyHandle> ValueHandle)
 	{
+		FText UnitDisplayText = FText::GetEmpty();
+		FName ParamName;
+		SetMetaData(StructPropertyHandle, ValueHandle, UnitDisplayText, ParamName);
+
 		const FText DisplayName = StructPropertyHandle->GetPropertyDisplayName();
 		FDetailWidgetRow& ValueNoModRow = ChildBuilder.AddCustomRow(DisplayName);
 		ValueNoModRow.NameContent()
