Minor: Move null impl to pure virtual for consistency

diff --git a/Engine/Source/Editor/AudioEditor/Public/AudioEditorModule.h b/Engine/Source/Editor/AudioEditor/Public/AudioEditorModule.h
--- a/Engine/Source/Editor/AudioEditor/Public/AudioEditorModule.h
+++ b/Engine/Source/Editor/AudioEditor/Public/AudioEditorModule.h
@@ -37,7 +37,7 @@ public:
 	virtual void RegisterAudioMixerAssetActions() = 0;
 
 	/** Registers effect preset asset actions. */
-	virtual void RegisterEffectPresetAssetActions() {}
+	virtual void RegisterEffectPresetAssetActions() = 0;
 
 	/** Creates a new sound class editor for a sound class object. */
 	virtual TSharedRef<FAssetEditorToolkit> CreateSoundClassEditor( const EToolkitMode::Type Mode, const TSharedPtr< IToolkitHost >& InitToolkitHost, USoundClass* InSoundClass ) = 0;
