Include/makeshared clean-up

diff --git a/Engine/Source/Editor/AudioEditor/Private/AudioEditorModule.cpp b/Engine/Source/Editor/AudioEditor/Private/AudioEditorModule.cpp
--- a/Engine/Source/Editor/AudioEditor/Private/AudioEditorModule.cpp
+++ b/Engine/Source/Editor/AudioEditor/Private/AudioEditorModule.cpp
@@ -1,49 +1,52 @@
 // Copyright Epic Games, Inc. All Rights Reserved.
 
 #include "AudioEditorModule.h"
-#include "Modules/ModuleManager.h"
-#include "Sound/SoundNodeDialoguePlayer.h"
-#include "EdGraphUtilities.h"
-#include "SoundCueGraphConnectionDrawingPolicy.h"
-#include "Factories/SoundFactory.h"
-#include "Factories/ReimportSoundFactory.h"
-#include "SoundCueGraph/SoundCueGraphNode.h"
-#include "SoundCueGraphNodeFactory.h"
+
 #include "AssetToolsModule.h"
-#include "PropertyEditorModule.h"
-#include "SoundClassEditor.h"
-#include "Sound/AudioSettings.h"
-#include "Sound/SoundCue.h"
-#include "Sound/SoundWave.h"
-#include "Sound/SoundSubmix.h"
-#include "Sound/SoundEffectPreset.h"
-#include "SoundCueEditor.h"
-#include "SoundModulationDestinationLayout.h"
-#include "SoundSubmixEditor.h"
+#include "AssetTypeActions/AssetTypeActions_AudioBus.h"
 #include "AssetTypeActions/AssetTypeActions_DialogueVoice.h"
 #include "AssetTypeActions/AssetTypeActions_DialogueWave.h"
+#include "AssetTypeActions/AssetTypeActions_ReverbEffect.h"
 #include "AssetTypeActions/AssetTypeActions_SoundAttenuation.h"
-#include "AssetTypeActions/AssetTypeActions_SoundConcurrency.h"
 #include "AssetTypeActions/AssetTypeActions_SoundBase.h"
 #include "AssetTypeActions/AssetTypeActions_SoundClass.h"
+#include "AssetTypeActions/AssetTypeActions_SoundConcurrency.h"
 #include "AssetTypeActions/AssetTypeActions_SoundCue.h"
+#include "AssetTypeActions/AssetTypeActions_SoundEffectPreset.h"
 #include "AssetTypeActions/AssetTypeActions_SoundMix.h"
 #include "AssetTypeActions/AssetTypeActions_SoundWave.h"
-#include "AssetTypeActions/AssetTypeActions_ReverbEffect.h"
-#include "AssetTypeActions/AssetTypeActions_SoundSubmix.h"
-#include "AssetTypeActions/AssetTypeActions_SoundEffectPreset.h"
 #include "AssetTypeActions/AssetTypeActions_SoundSourceBus.h"
-#include "AssetTypeActions/AssetTypeActions_AudioBus.h"
-#include "Utils.h"
-#include "UObject/UObjectIterator.h"
-#include "Styling/SlateStyle.h"
-#include "Styling/SlateStyleRegistry.h"
+#include "AssetTypeActions/AssetTypeActions_SoundSubmix.h"
+#include "ClassTemplateEditorSubsystem.h"
+#include "Components/SynthComponent.h"
+#include "EdGraphUtilities.h"
+#include "Factories/ReimportSoundFactory.h"
+#include "Factories/SoundFactory.h"
+#include "Modules/ModuleManager.h"
+#include "PropertyEditorModule.h"
+#include "Sound/AudioSettings.h"
+#include "Sound/SoundCue.h"
+#include "Sound/SoundEffectPreset.h"
+#include "Sound/SoundNodeDialoguePlayer.h"
+#include "Sound/SoundSubmix.h"
+#include "Sound/SoundWave.h"
+#include "SoundClassEditor.h"
+#include "SoundCueGraph/SoundCueGraphNode.h"
+#include "SoundCueGraphConnectionDrawingPolicy.h"
+#include "SoundCueGraphNodeFactory.h"
+#include "SoundCueEditor.h"
 #include "SoundFileIO/SoundFileIO.h"
+#include "SoundModulationDestinationLayout.h"
+#include "SoundSubmixEditor.h"
 #include "SoundSubmixGraph/SoundSubmixGraphSchema.h"
+#include "Styling/SlateStyle.h"
+#include "Styling/SlateStyleRegistry.h"
 #include "SubmixDetailsCustomization.h"
+#include "UObject/UObjectIterator.h"
+#include "Utils.h"
 #include "WidgetBlueprint.h"
 
-const FName AudioEditorAppIdentifier = FName(TEXT("AudioEditorApp"));
+const FName AudioEditorAppIdentifier("AudioEditorApp");
 
 DEFINE_LOG_CATEGORY(LogAudioEditor);
 
@@ -77,7 +80,7 @@
 	FAudioEditorModule()
 	{
 		// Create style set for audio asset icons
-		AudioStyleSet = MakeShareable(new FSlateStyleSet("AudioStyleSet"));
+		AudioStyleSet = MakeShared<FSlateStyleSet>("AudioStyleSet");
 	}
 
 	virtual void StartupModule() override
@@ -173,16 +176,16 @@
 		// Register the audio editor asset type actions
 		IAssetTools& AssetTools = FModuleManager::LoadModuleChecked<FAssetToolsModule>("AssetTools").Get();
 
-		AssetTools.RegisterAssetTypeActions(MakeShareable(new FAssetTypeActions_DialogueVoice));
-		AssetTools.RegisterAssetTypeActions(MakeShareable(new FAssetTypeActions_DialogueWave));
-		AssetTools.RegisterAssetTypeActions(MakeShareable(new FAssetTypeActions_SoundAttenuation));
-		AssetTools.RegisterAssetTypeActions(MakeShareable(new FAssetTypeActions_SoundConcurrency));
-		AssetTools.RegisterAssetTypeActions(MakeShareable(new FAssetTypeActions_SoundBase));
-		AssetTools.RegisterAssetTypeActions(MakeShareable(new FAssetTypeActions_SoundClass));
-		AssetTools.RegisterAssetTypeActions(MakeShareable(new FAssetTypeActions_SoundCue));
-		AssetTools.RegisterAssetTypeActions(MakeShareable(new FAssetTypeActions_SoundMix));
-		AssetTools.RegisterAssetTypeActions(MakeShareable(new FAssetTypeActions_SoundWave));
-		AssetTools.RegisterAssetTypeActions(MakeShareable(new FAssetTypeActions_ReverbEffect));
+		AssetTools.RegisterAssetTypeActions(MakeShared<FAssetTypeActions_DialogueVoice>());
+		AssetTools.RegisterAssetTypeActions(MakeShared<FAssetTypeActions_DialogueWave>());
+		AssetTools.RegisterAssetTypeActions(MakeShared<FAssetTypeActions_SoundAttenuation>());
+		AssetTools.RegisterAssetTypeActions(MakeShared<FAssetTypeActions_SoundConcurrency>());
+		AssetTools.RegisterAssetTypeActions(MakeShared<FAssetTypeActions_SoundBase>());
+		AssetTools.RegisterAssetTypeActions(MakeShared<FAssetTypeActions_SoundClass>());
+		AssetTools.RegisterAssetTypeActions(MakeShared<FAssetTypeActions_SoundCue>());
+		AssetTools.RegisterAssetTypeActions(MakeShared<FAssetTypeActions_SoundMix>());
+		AssetTools.RegisterAssetTypeActions(MakeShared<FAssetTypeActions_SoundWave>());
+		AssetTools.RegisterAssetTypeActions(MakeShared<FAssetTypeActions_ReverbEffect>());
 	}
 
 	virtual void RegisterAudioMixerAssetActions() override
@@ -191,20 +194,20 @@
 		if (GetDefault<UAudioSettings>()->IsAudioMixerEnabled())
 		{
 			IAssetTools& AssetTools = FModuleManager::LoadModuleChecked<FAssetToolsModule>("AssetTools").Get();
-			AssetTools.RegisterAssetTypeActions(MakeShareable(new FAssetTypeActions_SoundSubmix));
-			AssetTools.RegisterAssetTypeActions(MakeShareable(new FAssetTypeActions_SoundfieldSubmix));
-			AssetTools.RegisterAssetTypeActions(MakeShareable(new FAssetTypeActions_EndpointSubmix));
-			AssetTools.RegisterAssetTypeActions(MakeShareable(new FAssetTypeActions_SoundfieldEndpointSubmix));
-			AssetTools.RegisterAssetTypeActions(MakeShareable(new FAssetTypeActions_SoundfieldEncodingSettings));
-			AssetTools.RegisterAssetTypeActions(MakeShareable(new FAssetTypeActions_SoundfieldEffectSettings));
-			AssetTools.RegisterAssetTypeActions(MakeShareable(new FAssetTypeActions_SoundfieldEffect));
-			AssetTools.RegisterAssetTypeActions(MakeShareable(new FAssetTypeActions_AudioEndpointSettings));
-			AssetTools.RegisterAssetTypeActions(MakeShareable(new FAssetTypeActions_SoundfieldEndpointSettings));
-			AssetTools.RegisterAssetTypeActions(MakeShareable(new FAssetTypeActions_SoundEffectSubmixPreset));
-			AssetTools.RegisterAssetTypeActions(MakeShareable(new FAssetTypeActions_SoundEffectSourcePreset));
-			AssetTools.RegisterAssetTypeActions(MakeShareable(new FAssetTypeActions_SoundEffectSourcePresetChain));
-			AssetTools.RegisterAssetTypeActions(MakeShareable(new FAssetTypeActions_SoundSourceBus));
-			AssetTools.RegisterAssetTypeActions(MakeShareable(new FAssetTypeActions_AudioBus));
+			AssetTools.RegisterAssetTypeActions(MakeShared<FAssetTypeActions_SoundSubmix>());
+			AssetTools.RegisterAssetTypeActions(MakeShared<FAssetTypeActions_SoundfieldSubmix>());
+			AssetTools.RegisterAssetTypeActions(MakeShared<FAssetTypeActions_EndpointSubmix>());
+			AssetTools.RegisterAssetTypeActions(MakeShared<FAssetTypeActions_SoundfieldEndpointSubmix>());
+			AssetTools.RegisterAssetTypeActions(MakeShared<FAssetTypeActions_SoundfieldEncodingSettings>());
+			AssetTools.RegisterAssetTypeActions(MakeShared<FAssetTypeActions_SoundfieldEffectSettings>());
+			AssetTools.RegisterAssetTypeActions(MakeShared<FAssetTypeActions_SoundfieldEffect>());
+			AssetTools.RegisterAssetTypeActions(MakeShared<FAssetTypeActions_AudioEndpointSettings>());
+			AssetTools.RegisterAssetTypeActions(MakeShared<FAssetTypeActions_SoundfieldEndpointSettings>());
+			AssetTools.RegisterAssetTypeActions(MakeShared<FAssetTypeActions_SoundEffectSubmixPreset>());
+			AssetTools.RegisterAssetTypeActions(MakeShared<FAssetTypeActions_SoundEffectSourcePreset>());
+			AssetTools.RegisterAssetTypeActions(MakeShared<FAssetTypeActions_SoundEffectSourcePresetChain>());
+			AssetTools.RegisterAssetTypeActions(MakeShared<FAssetTypeActions_SoundSourceBus>());
+			AssetTools.RegisterAssetTypeActions(MakeShared<FAssetTypeActions_AudioBus>());
 		}
 	}
 
@@ -233,7 +236,7 @@
 					if (!RegisteredActions.Contains(EffectPreset) && EffectPreset->HasAssetActions())
 					{
 						RegisteredActions.Add(EffectPreset);
-						AssetTools.RegisterAssetTypeActions(MakeShareable(new FAssetTypeActions_SoundEffectPreset(EffectPreset)));
+						AssetTools.RegisterAssetTypeActions(MakeShared<FAssetTypeActions_SoundEffectPreset>(EffectPreset));
 					}
 				}
 			}
@@ -399,8 +402,8 @@
 
 		void Init()
 		{
-			MenuExtensibilityManager = MakeShareable(new FExtensibilityManager);
-			ToolBarExtensibilityManager = MakeShareable(new FExtensibilityManager);
+			MenuExtensibilityManager = MakeShared<FExtensibilityManager>();
+			ToolBarExtensibilityManager = MakeShared<FExtensibilityManager>();
 		}
 
 		void Reset()
