Add AudioEditorSubsystem and add functions to find and display custom UserWidgets generated from Blueprints implementing an IAudioWidgetInterface

diff --git a/Engine/Plugins/Runtime/Synthesis/Source/Synthesis/Classes/SubmixEffects/SubmixEffectDelay.h b/Engine/Plugins/Runtime/Synthesis/Source/Synthesis/Classes/SubmixEffects/SubmixEffectDelay.h
--- a/Engine/Plugins/Runtime/Synthesis/Source/Synthesis/Classes/SubmixEffects/SubmixEffectDelay.h
+++ b/Engine/Plugins/Runtime/Synthesis/Source/Synthesis/Classes/SubmixEffects/SubmixEffectDelay.h
@@ -6,6 +6,7 @@
 #include "AudioEffect.h"
 #include "DSP/Delay.h"
 #include "DSP/Dsp.h"
+#include "Kismet/BlueprintFunctionLibrary.h"
 #include "Sound/SoundEffectSubmix.h"
 #include "SubmixEffectDelay.generated.h"
 
@@ -39,6 +40,37 @@ struct SYNTHESIS_API FSubmixEffectDelaySettings
 	}
 };
 
+UCLASS()
+class SYNTHESIS_API USubmixEffectDelayStatics : public UBlueprintFunctionLibrary
+{
+	GENERATED_BODY()
+
+public:
+	UFUNCTION(BlueprintCallable, Category = "Audio|Effects|Delay")
+	static FSubmixEffectDelaySettings& SetMaximumDelayLength(UPARAM(ref) FSubmixEffectDelaySettings& DelaySettings, float MaximumDelayLength)
+	{
+		DelaySettings.MaximumDelayLength = FMath::Max(0.0f, MaximumDelayLength);
+		DelaySettings.DelayLength = FMath::Min(DelaySettings.MaximumDelayLength, DelaySettings.DelayLength);
+		return DelaySettings;
+	}
+
+	UFUNCTION(BlueprintCallable, Category = "Audio|Effects|Delay")
+	static FSubmixEffectDelaySettings& SetInterpolationTime(UPARAM(ref) FSubmixEffectDelaySettings& DelaySettings, float InterpolationTime)
+	{
+		DelaySettings.InterpolationTime = FMath::Max(0.0f, InterpolationTime);
+		return DelaySettings;
+	}
+
+	UFUNCTION(BlueprintCallable, Category = "Audio|Effects|Delay")
+	static FSubmixEffectDelaySettings& SetDelayLength(UPARAM(ref) FSubmixEffectDelaySettings& DelaySettings, float DelayLength)
+	{
+		DelaySettings.DelayLength = FMath::Max(0.0f, DelayLength);
+		DelaySettings.MaximumDelayLength = FMath::Max(DelaySettings.MaximumDelayLength, DelaySettings.DelayLength);
+
+		return DelaySettings;
+	}
+};
+
 class SYNTHESIS_API FSubmixEffectDelay : public FSoundEffectSubmix
 {
 public:
@@ -106,13 +138,22 @@ class SYNTHESIS_API USubmixEffectDelayPreset : public USoundEffectSubmixPreset
 public:
 	EFFECT_PRESET_METHODS(SubmixEffectDelay)
 
-	// Set all tap delay setting. This will replace any dynamically added or modified taps.
-	UFUNCTION(BlueprintCallable, Category = "Audio|Effects|Delay")
+	// Sets runtime delay settings. This will replace any dynamically added or modified settings without modifying
+	// the original UObject.
+	UFUNCTION(BlueprintCallable, Category = "Audio|Effects|Delay", meta = (DisplayName = "Set Dynamic Settings"))
 	void SetSettings(const FSubmixEffectDelaySettings& InSettings);
 
+	// Sets object's default settings. This will update both the default UObject settings (and mark it as dirty),
+	// as well as any dynamically set settings.
+	UFUNCTION(BlueprintCallable, Category = "Audio|Effects|Delay")
+	void SetDefaultSettings(const FSubmixEffectDelaySettings& InSettings);
+
 	// Get the maximum delay possible.
 	UFUNCTION(BlueprintCallable, Category = "Audio|Effects|Delay")
-	float GetMaxDelayInMilliseconds() { return DynamicSettings.MaximumDelayLength;  };
+	float GetMaxDelayInMilliseconds() const
+	{
+		return DynamicSettings.MaximumDelayLength;
+	};
 
 	// Set the time it takes to interpolate between parameters, in milliseconds.
 	UFUNCTION(BlueprintCallable, Category = "Audio|Effects|Delay")
@@ -125,7 +166,7 @@ public:
 public:
 	virtual void OnInit() override;
 
-	UPROPERTY(EditAnywhere, BlueprintReadWrite, Category = SubmixEffectPreset, Meta = (ShowOnlyInnerProperties))
+	UPROPERTY(EditAnywhere, BlueprintReadWrite, BlueprintSetter = SetDefaultSettings, Category = SubmixEffectPreset, Meta = (ShowOnlyInnerProperties))
 	FSubmixEffectDelaySettings Settings;
 
 	UPROPERTY(transient)
diff --git a/Engine/Plugins/Runtime/Synthesis/Source/Synthesis/Private/SubmixEffectDelay.cpp b/Engine/Plugins/Runtime/Synthesis/Source/Synthesis/Private/SubmixEffectDelay.cpp
--- a/Engine/Plugins/Runtime/Synthesis/Source/Synthesis/Private/SubmixEffectDelay.cpp
+++ b/Engine/Plugins/Runtime/Synthesis/Source/Synthesis/Private/SubmixEffectDelay.cpp
@@ -160,7 +160,22 @@ void USubmixEffectDelayPreset::OnInit()
 void USubmixEffectDelayPreset::SetSettings(const FSubmixEffectDelaySettings& InSettings)
 {
 	DynamicSettings = InSettings;
-	UpdateSettings(InSettings);
+
+// 	USubmixEffectDelayStatics::SetDelayLength(DynamicSettings, DynamicSettings.DelayLength);
+// 	USubmixEffectDelayStatics::SetMaximumDelayLength(DynamicSettings, DynamicSettings.MaximumDelayLength);
+
+	UpdateSettings(DynamicSettings);
 }
 
+void USubmixEffectDelayPreset::SetDefaultSettings(const FSubmixEffectDelaySettings& InSettings)
+{
+	Settings = InSettings;
+
+// 	USubmixEffectDelayStatics::SetDelayLength(Settings, Settings.DelayLength);
+// 	USubmixEffectDelayStatics::SetMaximumDelayLength(Settings, Settings.MaximumDelayLength);
 
+	DynamicSettings = Settings;
+	UpdateSettings(DynamicSettings);
+
+	MarkPackageDirty();
+}
diff --git a/Engine/Plugins/Runtime/Synthesis/Source/SynthesisEditor/Classes/SynthesisEditorSettings.h b/Engine/Plugins/Runtime/Synthesis/Source/SynthesisEditor/Classes/SynthesisEditorSettings.h
--- a/Engine/Plugins/Runtime/Synthesis/Source/SynthesisEditor/Classes/SynthesisEditorSettings.h
+++ b/Engine/Plugins/Runtime/Synthesis/Source/SynthesisEditor/Classes/SynthesisEditorSettings.h
@@ -12,67 +12,4 @@ UCLASS(config = EditorSettings, defaultconfig, meta = (DisplayName = "Synthesis"
 class USynthesisEditorSettings : public UDeveloperSettings
 {
 	GENERATED_BODY()
-
-public:
-	UPROPERTY(config, EditAnywhere, Category = "Source Effects", meta = (AllowedClasses = "WidgetBlueprint"))
-	FSoftObjectPath BitCrusherWidget;
-
-	UPROPERTY(config, EditAnywhere, Category = "Source Effects", meta = (AllowedClasses = "WidgetBlueprint"))
-	FSoftObjectPath ChorusWidget;
-
-	UPROPERTY(config, EditAnywhere, Category = "Source Effects", meta = (AllowedClasses = "WidgetBlueprint"))
-	FSoftObjectPath DynamicsProcessorWidget;
-
-	UPROPERTY(config, EditAnywhere, Category = "Source Effects", meta = (AllowedClasses = "WidgetBlueprint"))
-	FSoftObjectPath EnvelopeFollowerWidget;
-
-	UPROPERTY(config, EditAnywhere, Category = "Source Effects", meta = (AllowedClasses = "WidgetBlueprint"))
-	FSoftObjectPath EQWidget;
-
-	UPROPERTY(config, EditAnywhere, Category = "Source Effects", meta = (DisplayName = "Filter Widget (Source)", AllowedClasses = "WidgetBlueprint"))
-	FSoftObjectPath FilterWidget;
-
-	UPROPERTY(config, EditAnywhere, Category = "Source Effects", meta = (AllowedClasses = "WidgetBlueprint"))
-	FSoftObjectPath FoldbackDistortionWidget;
-
-	UPROPERTY(config, EditAnywhere, Category = "Source Effects", meta = (AllowedClasses = "WidgetBlueprint"))
-	FSoftObjectPath MidSideSpreaderWidget;
-
-	UPROPERTY(config, EditAnywhere, Category = "Source Effects", meta = (AllowedClasses = "WidgetBlueprint"))
-	FSoftObjectPath PannerWidget;
-
-	UPROPERTY(config, EditAnywhere, Category = "Source Effects", meta = (AllowedClasses = "WidgetBlueprint"))
-	FSoftObjectPath PhaserWidget;
-
-	UPROPERTY(config, EditAnywhere, Category = "Source Effects", meta = (AllowedClasses = "WidgetBlueprint"))
-	FSoftObjectPath RingModulationWidget;
-
-	UPROPERTY(config, EditAnywhere, Category = "Source Effects", meta = (AllowedClasses = "WidgetBlueprint"))
-	FSoftObjectPath SimpleDelayWidget;
-
-	UPROPERTY(config, EditAnywhere, Category = "Source Effects", meta = (AllowedClasses = "WidgetBlueprint"))
-	FSoftObjectPath StereoDelayWidget;
-
-	UPROPERTY(config, EditAnywhere, Category = "Source Effects", meta = (AllowedClasses = "WidgetBlueprint"))
-	FSoftObjectPath WaveShaperWidget;
-
-	UPROPERTY(config, EditAnywhere, Category = "Submix Effects", meta = (DisplayName = "Convolution Reverb Widget", AllowedClasses = "WidgetBlueprint"))
-	FSoftObjectPath SubmixConvolutionReverbWidget;
-
-	UPROPERTY(config, EditAnywhere, Category = "Submix Effects", meta = (DisplayName = "Delay Widget (Submix)", AllowedClasses = "WidgetBlueprint"))
-	FSoftObjectPath SubmixDelayWidget;
-
-	UPROPERTY(config, EditAnywhere, Category = "Submix Effects", meta = (DisplayName = "Filter Widget (Submix)", AllowedClasses = "WidgetBlueprint"))
-	FSoftObjectPath SubmixFilterWidget;
-
-	UPROPERTY(config, EditAnywhere, Category = "Submix Effects", meta = (DisplayName = "Flexiverb Widget", AllowedClasses = "WidgetBlueprint"))
-	FSoftObjectPath SubmixFlexiverbWidget;
-
-	UPROPERTY(config, EditAnywhere, Category = "Submix Effects", meta = (DisplayName = "Stereo Delay (Submix)", AllowedClasses = "WidgetBlueprint"))
-	FSoftObjectPath SubmixStereoDelayWidget;
-
-	UPROPERTY(config, EditAnywhere, Category = "Submix Effects", meta = (AllowedClasses = "WidgetBlueprint"))
-	FSoftObjectPath SubmixTapDelayWidget;
-
-	void PostEditChangeProperty(FPropertyChangedEvent& PropertyChangedEvent);
 };
diff --git a/Engine/Plugins/Runtime/Synthesis/Source/SynthesisEditor/Private/SynthesisEditorModule.cpp b/Engine/Plugins/Runtime/Synthesis/Source/SynthesisEditor/Private/SynthesisEditorModule.cpp
--- a/Engine/Plugins/Runtime/Synthesis/Source/SynthesisEditor/Private/SynthesisEditorModule.cpp
+++ b/Engine/Plugins/Runtime/Synthesis/Source/SynthesisEditor/Private/SynthesisEditorModule.cpp
@@ -45,56 +45,6 @@ DEFINE_LOG_CATEGORY(LogSynthesisEditor);
 
 IMPLEMENT_MODULE(FSynthesisEditorModule, SynthesisEditor)
 
-namespace SynthesisEditorUtils
-{
-	void RegisterSoundEffectPresets(const FName* InPropertyChanged = nullptr)
-	{
-		IAudioEditorModule* AudioEditorModule = &FModuleManager::LoadModuleChecked<IAudioEditorModule>("AudioEditor");
-		const USynthesisEditorSettings* Settings = GetDefault<USynthesisEditorSettings>();
-		check(Settings);
-
-		auto RegisterWidget = [AudioEditorModule, InPropertyChanged](FName InName, FSoftObjectPath InPath, TSubclassOf<USoundEffectPreset> InClass)
-		{
-			if (!InPropertyChanged || *InPropertyChanged == InName)
-			{
-				if (UWidgetBlueprint* WidgetBP = Cast<UWidgetBlueprint>(InPath.TryLoad()))
-				{
-					AudioEditorModule->RegisterSoundEffectPresetWidget(InClass, WidgetBP);
-				}
-			}
-		};
-
-		// Source Effects
-		RegisterWidget(GET_MEMBER_NAME_CHECKED(USynthesisEditorSettings, BitCrusherWidget), Settings->BitCrusherWidget, USourceEffectBitCrusherPreset::StaticClass());
-		RegisterWidget(GET_MEMBER_NAME_CHECKED(USynthesisEditorSettings, ChorusWidget), Settings->ChorusWidget, USourceEffectChorusPreset::StaticClass());
-		RegisterWidget(GET_MEMBER_NAME_CHECKED(USynthesisEditorSettings, DynamicsProcessorWidget), Settings->DynamicsProcessorWidget, USourceEffectDynamicsProcessorPreset::StaticClass());
-		RegisterWidget(GET_MEMBER_NAME_CHECKED(USynthesisEditorSettings, EnvelopeFollowerWidget), Settings->EnvelopeFollowerWidget, USourceEffectEnvelopeFollowerPreset::StaticClass());
-		RegisterWidget(GET_MEMBER_NAME_CHECKED(USynthesisEditorSettings, EQWidget), Settings->EQWidget, USourceEffectEQPreset::StaticClass());
-		RegisterWidget(GET_MEMBER_NAME_CHECKED(USynthesisEditorSettings, FilterWidget), Settings->FilterWidget, USourceEffectFilterPreset::StaticClass());
-		RegisterWidget(GET_MEMBER_NAME_CHECKED(USynthesisEditorSettings, FoldbackDistortionWidget), Settings->FoldbackDistortionWidget, USourceEffectFoldbackDistortionPreset::StaticClass());
-		RegisterWidget(GET_MEMBER_NAME_CHECKED(USynthesisEditorSettings, MidSideSpreaderWidget), Settings->MidSideSpreaderWidget, USourceEffectMidSideSpreaderPreset::StaticClass());
-		RegisterWidget(GET_MEMBER_NAME_CHECKED(USynthesisEditorSettings, PannerWidget), Settings->PannerWidget, USourceEffectPannerPreset::StaticClass());
-		RegisterWidget(GET_MEMBER_NAME_CHECKED(USynthesisEditorSettings, PhaserWidget), Settings->PhaserWidget, USourceEffectPhaserPreset::StaticClass());
-		RegisterWidget(GET_MEMBER_NAME_CHECKED(USynthesisEditorSettings, RingModulationWidget), Settings->RingModulationWidget, USourceEffectRingModulationPreset::StaticClass());
-		RegisterWidget(GET_MEMBER_NAME_CHECKED(USynthesisEditorSettings, SimpleDelayWidget), Settings->SimpleDelayWidget, USourceEffectSimpleDelayPreset::StaticClass());
-		RegisterWidget(GET_MEMBER_NAME_CHECKED(USynthesisEditorSettings, StereoDelayWidget), Settings->StereoDelayWidget, USourceEffectStereoDelayPreset::StaticClass());
-		RegisterWidget(GET_MEMBER_NAME_CHECKED(USynthesisEditorSettings, WaveShaperWidget), Settings->WaveShaperWidget, USourceEffectWaveShaperPreset::StaticClass());
-
-		// Submix Effects
-		RegisterWidget(GET_MEMBER_NAME_CHECKED(USynthesisEditorSettings, SubmixConvolutionReverbWidget), Settings->SubmixConvolutionReverbWidget, USubmixEffectConvolutionReverbPreset::StaticClass());
-		RegisterWidget(GET_MEMBER_NAME_CHECKED(USynthesisEditorSettings, SubmixDelayWidget), Settings->SubmixDelayWidget, USubmixEffectDelayPreset::StaticClass());
-		RegisterWidget(GET_MEMBER_NAME_CHECKED(USynthesisEditorSettings, SubmixFilterWidget), Settings->SubmixFilterWidget, USubmixEffectFilterPreset::StaticClass());
-		RegisterWidget(GET_MEMBER_NAME_CHECKED(USynthesisEditorSettings, SubmixFlexiverbWidget), Settings->SubmixFlexiverbWidget, USubmixEffectFlexiverbPreset::StaticClass());
-		RegisterWidget(GET_MEMBER_NAME_CHECKED(USynthesisEditorSettings, SubmixStereoDelayWidget), Settings->SubmixStereoDelayWidget, USubmixEffectStereoDelayPreset::StaticClass());
-		RegisterWidget(GET_MEMBER_NAME_CHECKED(USynthesisEditorSettings, SubmixTapDelayWidget), Settings->SubmixTapDelayWidget, USubmixEffectTapDelayPreset::StaticClass());
-	}
-} // namespace SynthesisEditorUtils
-
-void USynthesisEditorSettings::PostEditChangeProperty(FPropertyChangedEvent& PropertyChangedEvent)
-{
-	const FName PropertyName = PropertyChangedEvent.GetPropertyName();
-	SynthesisEditorUtils::RegisterSoundEffectPresets(&PropertyName);
-}
 
 void FSynthesisEditorModule::StartupModule()
 {
@@ -110,8 +60,6 @@ void FSynthesisEditorModule::StartupModule()
 	IAudioEditorModule* AudioEditorModule = &FModuleManager::LoadModuleChecked<IAudioEditorModule>("AudioEditor");
 	AudioEditorModule->RegisterEffectPresetAssetActions();
 
-	SynthesisEditorUtils::RegisterSoundEffectPresets();
-
 	UToolMenus::RegisterStartupCallback(FSimpleMulticastDelegate::FDelegate::CreateRaw(this, &FSynthesisEditorModule::RegisterMenus));
 }
 
diff --git a/Engine/Source/Editor/AudioEditor/Private/AssetTypeActions/AssetTypeActions_SoundEffectPreset.cpp b/Engine/Source/Editor/AudioEditor/Private/AssetTypeActions/AssetTypeActions_SoundEffectPreset.cpp
--- a/Engine/Source/Editor/AudioEditor/Private/AssetTypeActions/AssetTypeActions_SoundEffectPreset.cpp
+++ b/Engine/Source/Editor/AudioEditor/Private/AssetTypeActions/AssetTypeActions_SoundEffectPreset.cpp
@@ -1,6 +1,7 @@
 // Copyright Epic Games, Inc. All Rights Reserved.
 
 #include "AssetTypeActions_SoundEffectPreset.h"
+#include "AudioEditorSubsystem.h"
 #include "Developer/AssetTools/Public/IAssetTools.h"
 #include "Developer/AssetTools/Public/AssetToolsModule.h"
 #include "Editors/SoundEffectPresetEditor.h"
@@ -11,17 +12,6 @@
 
 #define LOCTEXT_NAMESPACE "AssetTypeActions"
 
-static bool bPrototypeSFXEditorEnabled = false;
-static FAutoConsoleCommand GEnableSoundEffectEditorPrototype(
-	TEXT("au.AudioEditor.EnableSoundEffectEditorPrototype"),
-	TEXT("Enables's the UE5 prototype sound effect editor.\n"),
-	FConsoleCommandWithArgsDelegate::CreateStatic(
-		[](const TArray<FString>& Args)
-		{
-			bPrototypeSFXEditorEnabled = true;
-		}
-	)
-);
 
 namespace EffectPresets
 {
@@ -91,22 +81,30 @@ const TArray<FText>& FAssetTypeActions_SoundEffectSourcePreset::GetSubMenus() co
 
 void FAssetTypeActions_SoundEffectPreset::OpenAssetEditor(const TArray<UObject*>& InObjects, TSharedPtr<IToolkitHost> ToolkitHost)
 {
-	if (!bPrototypeSFXEditorEnabled)
-	{
-		FAssetTypeActions_Base::OpenAssetEditor(InObjects, ToolkitHost);
-	}
-	else
-	{
-		EToolkitMode::Type Mode = ToolkitHost.IsValid() ? EToolkitMode::WorldCentric : EToolkitMode::Standalone;
+	EToolkitMode::Type Mode = ToolkitHost.IsValid() ? EToolkitMode::WorldCentric : EToolkitMode::Standalone;
 
+	TArray<UObject*> Objects = InObjects;
+
+	if (UAudioEditorSubsystem* AudioEditorSubsytem = GEditor->GetEditorSubsystem<UAudioEditorSubsystem>())
+	{
 		for (UObject* Object : InObjects)
 		{
 			if (USoundEffectPreset* Preset = Cast<USoundEffectPreset>(Object))
 			{
-				TSharedRef<FSoundEffectPresetEditor> PresetEditor = MakeShared<FSoundEffectPresetEditor>();
-				PresetEditor->Init(Mode, ToolkitHost, Preset);
+				TArray<UUserWidget*> UserWidgets = AudioEditorSubsytem->CreateUserWidgets(USoundEffectPresetViewInterface::StaticClass(), Object->GetClass());
+				if (!UserWidgets.IsEmpty())
+				{
+					TSharedRef<FSoundEffectPresetEditor> PresetEditor = MakeShared<FSoundEffectPresetEditor>();
+					PresetEditor->Init(Mode, ToolkitHost, Preset, UserWidgets);
+					Objects.Remove(Object);
+				}
 			}
 		}
 	}
+
+	if (!Objects.IsEmpty())
+	{
+		FAssetTypeActions_Base::OpenAssetEditor(Objects, ToolkitHost);
+	}
 }
 #undef LOCTEXT_NAMESPACE
\ No newline at end of file
diff --git a/Engine/Source/Editor/AudioEditor/Private/AudioEditorModule.cpp b/Engine/Source/Editor/AudioEditor/Private/AudioEditorModule.cpp
--- a/Engine/Source/Editor/AudioEditor/Private/AudioEditorModule.cpp
+++ b/Engine/Source/Editor/AudioEditor/Private/AudioEditorModule.cpp
@@ -296,35 +296,6 @@ public:
 		return SoundCueExtensibility.MenuExtensibilityManager;
 	}
 
-	virtual void RegisterSoundEffectPresetWidget(TSubclassOf<USoundEffectPreset> PresetClass, UWidgetBlueprint* WidgetBlueprint) override
-	{
-		UnregisterSoundEffectPresetWidget(PresetClass);
-
-		if (PresetClass)
-		{
-			WidgetBlueprint->AddToRoot();
-			EffectPresetWidgets.Add(PresetClass, WidgetBlueprint);
-		}
-	}
-
-	/** Returns custom widget blueprint for a given SoundEffectPreset class (or null if unset). */
-	virtual UWidgetBlueprint* GetSoundEffectPresetWidget(TSubclassOf<USoundEffectPreset> PresetClass) override
-	{
-		return EffectPresetWidgets.FindRef(PresetClass);
-	}
-
-	virtual void UnregisterSoundEffectPresetWidget(TSubclassOf<USoundEffectPreset> PresetClass) override
-	{
-		if (PresetClass)
-		{
-			if (UWidgetBlueprint* WidgetBlueprint = EffectPresetWidgets.FindRef(PresetClass))
-			{
-				WidgetBlueprint->RemoveFromRoot();
-				EffectPresetWidgets.Remove(PresetClass);
-			}
-		}
-	}
-
 	virtual void ReplaceSoundNodesInGraph(USoundCue* SoundCue, UDialogueWave* DialogueWave, TArray<USoundNode*>& NodesToReplace, const FDialogueContextMapping& ContextMapping) override
 	{
 		// Replace any sound nodes in the graph.
diff --git a/Engine/Source/Editor/AudioEditor/Private/AudioEditorSubsystem.cpp b/Engine/Source/Editor/AudioEditor/Private/AudioEditorSubsystem.cpp
new file mode 100644
--- /dev/null
+++ b/Engine/Source/Editor/AudioEditor/Private/AudioEditorSubsystem.cpp
@@ -0,0 +1,125 @@
+// Copyright Epic Games, Inc. All Rights Reserved.
+#include "AudioEditorSubsystem.h"
+
+#include "AssetRegistry/AssetRegistryModule.h"
+#include "Blueprint/BlueprintSupport.h"
+#include "Blueprint/UserWidget.h"
+#include "Editor.h"
+#include "Interfaces/IPluginManager.h"
+#include "Internationalization/Text.h"
+#include "Misc/CoreDelegates.h"
+#include "Misc/FileHelper.h"
+#include "Modules/ModuleManager.h"
+#include "Subsystems/SubsystemCollection.h"
+#include "UObject/CoreRedirects.h"
+#include "WidgetBlueprint.h"
+
+
+#define LOCTEXT_NAMESPACE "AudioEditorSubsystem"
+
+TArray<UUserWidget*> UAudioEditorSubsystem::CreateUserWidgets(TSubclassOf<UAudioWidgetInterface> InWidgetClass, UClass* InObjectClass) const
+{
+	TArray<UUserWidget*> UserWidgets;
+	UClass* InterfaceClass = InWidgetClass ? InWidgetClass.Get() : UAudioWidgetInterface::StaticClass();
+	if (!ensure(InterfaceClass))
+	{
+		return UserWidgets;
+	}
+
+	UWorld* World = GEditor ? GEditor->GetEditorWorldContext().World() : nullptr;
+	if (!World)
+	{
+		return UserWidgets;
+	}
+
+	const TArray<FAssetData> AssetData = GetWidgetBlueprintAssetData();
+	for (const FAssetData& AssetEntry : AssetData)
+	{
+		if (!ImplementsInterface(AssetEntry, InterfaceClass))
+		{
+			continue;
+		}
+
+		UObject* Asset = AssetEntry.GetAsset();
+		if (!Asset)
+		{
+			continue;
+		}
+
+		UWidgetBlueprint* WidgetBlueprint = CastChecked<UWidgetBlueprint>(Asset);
+		UClass* GeneratedClass = WidgetBlueprint->GeneratedClass;
+		if (!GeneratedClass)
+		{
+			continue;
+		}
+
+		UUserWidget* UserWidget = CreateWidget<UUserWidget>(World, GeneratedClass);
+		if (!UserWidget)
+		{
+			continue;
+		}
+
+		if (InObjectClass)
+		{
+			const UClass* ObjectClass = IAudioWidgetInterface::Execute_GetClass(UserWidget);
+			while (ObjectClass)
+			{
+				if (ObjectClass == InObjectClass)
+				{
+					UserWidgets.Add(UserWidget);
+					break;
+				}
+
+				ObjectClass = ObjectClass->GetSuperClass();
+			}
+		}
+		else
+		{
+			UserWidgets.Add(UserWidget);
+		}
+	}
+
+	return UserWidgets;
+}
+
+const TArray<FAssetData> UAudioEditorSubsystem::GetWidgetBlueprintAssetData()
+{
+	FAssetRegistryModule& AssetRegistryModule = FModuleManager::LoadModuleChecked<FAssetRegistryModule>("AssetRegistry");
+	TArray<FAssetData> AssetData;
+	const UClass* Class = UWidgetBlueprint::StaticClass();
+	AssetRegistryModule.Get().GetAssetsByClass(Class->GetFName(), AssetData);
+	return AssetData;
+}
+
+bool UAudioEditorSubsystem::ImplementsInterface(const FAssetData& InAssetData, UClass* InInterfaceClass)
+{
+	const FString ImplementedInterfaces = InAssetData.GetTagValueRef<FString>(FBlueprintTags::ImplementedInterfaces);
+	if (!ImplementedInterfaces.IsEmpty())
+	{
+		FString FullInterface;
+		FString RemainingString;
+		FString InterfacePath;
+		FString CurrentString = *ImplementedInterfaces;
+		while (CurrentString.Split(TEXT(","), &FullInterface, &RemainingString))
+		{
+			if (!CurrentString.StartsWith(TEXT("Graphs=(")))
+			{
+				if (FullInterface.Split(TEXT("\""), &CurrentString, &InterfacePath, ESearchCase::CaseSensitive))
+				{
+					InterfacePath.RemoveFromEnd(TEXT("\"'"));
+
+					const FCoreRedirectObjectName ResolvedInterfaceName = FCoreRedirects::GetRedirectedName(ECoreRedirectFlags::Type_Class, FCoreRedirectObjectName(InterfacePath));
+					if (InInterfaceClass->GetFName() == ResolvedInterfaceName.ObjectName)
+					{
+						return true;
+					}
+				}
+			}
+
+			CurrentString = RemainingString;
+		}
+	}
+
+	return false;
+}
+#undef LOCTEXT_NAMESPACE // AudioEditorSubsystem
\ No newline at end of file
diff --git a/Engine/Source/Editor/AudioEditor/Private/Editors/SoundEffectPresetEditor.cpp b/Engine/Source/Editor/AudioEditor/Private/Editors/SoundEffectPresetEditor.cpp
--- a/Engine/Source/Editor/AudioEditor/Private/Editors/SoundEffectPresetEditor.cpp
+++ b/Engine/Source/Editor/AudioEditor/Private/Editors/SoundEffectPresetEditor.cpp
@@ -24,13 +24,12 @@
 #define LOCTEXT_NAMESPACE "SoundEffectPresetEditor"
 
 
-const FName FSoundEffectPresetEditor::AppIdentifier(TEXT("SoundEffectPresetEditorApp"));
-const FName FSoundEffectPresetEditor::PropertiesTabId(TEXT("SoundEffectPresetEditor_Properties"));
-const FName FSoundEffectPresetEditor::UserWidgetTabId(TEXT("SoundEffectPresetEditor_UserWidget"));
+const FName FSoundEffectPresetEditor::AppIdentifier("SoundEffectPresetEditorApp");
+const FName FSoundEffectPresetEditor::PropertiesTabId("SoundEffectPresetEditor_Properties");
+const FName FSoundEffectPresetEditor::UserWidgetTabId("SoundEffectPresetEditor_UserWidget");
 
 FSoundEffectPresetEditor::FSoundEffectPresetEditor()
 	: SoundEffectPreset(nullptr)
-	, UserWidget(nullptr)
 {
 }
 
@@ -46,11 +45,14 @@ void FSoundEffectPresetEditor::RegisterTabSpawners(const TSharedRef<FTabManager>
 		.SetIcon(FSlateIcon(FEditorStyle::GetStyleSetName(), "LevelEditor.Tabs.Details"));
 
 
-	if (UserWidget)
+	for (int32 i = 0; i < UserWidgets.Num(); i++)
 	{
+		TStrongObjectPtr<UUserWidget> UserWidget = UserWidgets[i];
 		const FString ClassName = SoundEffectPreset->GetClass()->GetName();
 		FSlateIcon BPIcon(FEditorStyle::GetStyleSetName(), "LevelEditor.CreateClassBlueprint");
-		InTabManager->RegisterTabSpawner(UserWidgetTabId, FOnSpawnTab::CreateLambda([this](const FSpawnTabArgs& Args) { return SpawnTab_UserWidgetEditor(Args); }))
+
+		const FName UserWidgetTabIdIndexed = FName(UserWidgetTabId.ToString() + FString(TEXT("_")) + FString::FromInt(i));
+		InTabManager->RegisterTabSpawner(UserWidgetTabIdIndexed, FOnSpawnTab::CreateLambda([this, i](const FSpawnTabArgs& Args) { return SpawnTab_UserWidgetEditor(Args, i); }))
 			.SetDisplayName(FText::Format(LOCTEXT("UserEditorTabFormat", "{0} Editor"), FText::FromString(ClassName)))
 			.SetGroup(WorkspaceMenuCategory.ToSharedRef())
 			.SetIcon(BPIcon);
@@ -61,22 +63,30 @@ void FSoundEffectPresetEditor::UnregisterTabSpawners(const TSharedRef<FTabManage
 {
 	InTabManager->UnregisterTabSpawner(PropertiesTabId);
 
-	if (UserWidget)
+	for (int32 i = 0; i < UserWidgets.Num(); i++)
 	{
-		InTabManager->UnregisterTabSpawner(UserWidgetTabId);
+		const FName UserWidgetTabIdIndexed = FName(UserWidgetTabId.ToString() + FString(TEXT("_")) + FString::FromInt(i));
+		InTabManager->UnregisterTabSpawner(UserWidgetTabIdIndexed);
 	}
 }
 
-void FSoundEffectPresetEditor::Init(const EToolkitMode::Type Mode, const TSharedPtr<IToolkitHost>& InitToolkitHost, USoundEffectPreset* InPresetToEdit)
+void FSoundEffectPresetEditor::Init(const EToolkitMode::Type Mode, const TSharedPtr<IToolkitHost>& InitToolkitHost, USoundEffectPreset* InPresetToEdit, const TArray<UUserWidget*>& InWidgetBlueprints)
 {
-	check(InPresetToEdit);
-
-	SoundEffectPreset = InPresetToEdit;
-	InitPresetWidget();
+	if (!ensure(InPresetToEdit))
+	{
+		return;
+	}
 
+	SoundEffectPreset = TStrongObjectPtr<USoundEffectPreset>(InPresetToEdit);
+	InitPresetWidgets(InWidgetBlueprints);
+	
 	// Support undo/redo
 	InPresetToEdit->SetFlags(RF_Transactional);
-	GEditor->RegisterForUndo(this);
+
+	if (GEditor)
+	{
+		GEditor->RegisterForUndo(this);
+	}
 
 	FDetailsViewArgs Args;
 	Args.bHideSelectionTip = true;
@@ -96,7 +106,7 @@ void FSoundEffectPresetEditor::Init(const EToolkitMode::Type Mode, const TShared
 			->AddTab(PropertiesTabId, ETabState::OpenedTab)
 		);
 
-	if (UserWidget)
+	if (!UserWidgets.IsEmpty())
 	{
 		TabSplitter->Split
 		(
@@ -144,6 +154,11 @@ void FSoundEffectPresetEditor::Init(const EToolkitMode::Type Mode, const TShared
 		bUseSmallIcons);
 }
 
+FName FSoundEffectPresetEditor::GetEditorName() const
+{
+	return "Preset Editor";
+}
+
 FName FSoundEffectPresetEditor::GetToolkitFName() const
 {
 	return FName("SoundEffectPresetEditor");
@@ -151,10 +166,10 @@ FName FSoundEffectPresetEditor::GetToolkitFName() const
 
 FText FSoundEffectPresetEditor::GetBaseToolkitName() const
 {
-	return LOCTEXT("AppLabel", "Sound Effect Editor");
+	return LOCTEXT("AppLabel", "Sound Effect Preset Editor");
 }
 
-void FSoundEffectPresetEditor::InitPresetWidget()
+void FSoundEffectPresetEditor::InitPresetWidgets(const TArray<UUserWidget*>& InWidgets)
 {
 	if (!SoundEffectPreset)
 	{
@@ -163,15 +178,12 @@ void FSoundEffectPresetEditor::InitPresetWidget()
 
 	if (UWorld* World = GEditor->GetEditorWorldContext().World())
 	{
-		IAudioEditorModule* AudioEditorModule = &FModuleManager::LoadModuleChecked<IAudioEditorModule>("AudioEditor");
-
-		UClass* PresetClass = SoundEffectPreset->GetClass();
-		if (const UWidgetBlueprint* WidgetBP = AudioEditorModule->GetSoundEffectPresetWidget(PresetClass))
+		for (UUserWidget* Widget : InWidgets)
 		{
-			if (UClass* GeneratedClass = WidgetBP->GeneratedClass.Get())
+			if (Widget)
 			{
-				UserWidget = CreateWidget<USoundEffectPresetUserWidget>(World, GeneratedClass);
-				UserWidget->Preset = SoundEffectPreset;
+				UserWidgets.Add(TStrongObjectPtr<UUserWidget>(Widget));
+				IAudioWidgetInterface::Execute_OnConstructed(Widget, Cast<UObject>(SoundEffectPreset.Get()));
 			}
 		}
 	}
@@ -196,22 +208,34 @@ EOrientation FSoundEffectPresetEditor::GetSnapLabelOrientation() const
 
 void FSoundEffectPresetEditor::NotifyPostChange(const FPropertyChangedEvent& PropertyChangedEvent, FProperty* PropertyThatChanged)
 {
-	if (UserWidget && SoundEffectPreset)
+	if (SoundEffectPreset)
 	{
-		const FName PropertyName = PropertyThatChanged->GetFName();
-		UserWidget->OnPresetChanged(PropertyName);
+		for (TStrongObjectPtr<UUserWidget>& UserWidget : UserWidgets)
+		{
+			const FName PropertyName = PropertyThatChanged->GetFName();
+			IAudioWidgetInterface::Execute_OnPropertyChanged(UserWidget.Get(), SoundEffectPreset.Get(), PropertyName);
+		}
 	}
 }
 
-void FSoundEffectPresetEditor::AddReferencedObjects(FReferenceCollector& Collector)
+void FSoundEffectPresetEditor::NotifyPostChange(const FPropertyChangedEvent& PropertyChangedEvent, FEditPropertyChain* PropertyThatChanged)
 {
-	Collector.AddReferencedObject(SoundEffectPreset);
-	Collector.AddReferencedObject(UserWidget);
-}
-
-FString FSoundEffectPresetEditor::GetReferencerName() const
-{
-	return "SoundEffectPresetEditor";
+	if (SoundEffectPreset)
+	{
+		for (TStrongObjectPtr<UUserWidget>& UserWidget : UserWidgets)
+		{
+			auto Node = PropertyThatChanged->GetHead();
+			while(Node)
+			{
+				if (FProperty* Property = Node->GetValue())
+				{
+					const FName PropertyName = Property->GetFName();
+					IAudioWidgetInterface::Execute_OnPropertyChanged(UserWidget.Get(), SoundEffectPreset.Get(), PropertyName);
+				}
+				Node = Node->GetNextNode();
+			}
+		}
+	}
 }
 
 TSharedRef<SDockTab> FSoundEffectPresetEditor::SpawnTab_Properties(const FSpawnTabArgs& Args)
@@ -226,12 +250,18 @@ TSharedRef<SDockTab> FSoundEffectPresetEditor::SpawnTab_Properties(const FSpawnT
 		];
 }
 
-TSharedRef<SDockTab> FSoundEffectPresetEditor::SpawnTab_UserWidgetEditor(const FSpawnTabArgs& Args)
+TSharedRef<SDockTab> FSoundEffectPresetEditor::SpawnTab_UserWidgetEditor(const FSpawnTabArgs& Args, int32 WidgetIndex)
 {
+	FName IconBrushName = IAudioWidgetInterface::Execute_GetIconBrushName(UserWidgets[WidgetIndex].Get());
+	if (IconBrushName == FName())
+	{
+		IconBrushName = ("GenericEditor.Tabs.Properties");
+	}
+
 	const FSlateBrush* IconBrush = FEditorStyle::GetBrush("SoundEffectPresetEditor.Tabs.Properties");
-	const FText Label = FText::FromString(GetEditingObject()->GetName());
 
-	if (!UserWidget)
+	FText Label = FText::FromString(SoundEffectPreset->GetName());
+	if (UserWidgets.Num() < WidgetIndex)
 	{
 		return SNew(SDockTab)
 			.Icon(IconBrush)
@@ -243,6 +273,12 @@ TSharedRef<SDockTab> FSoundEffectPresetEditor::SpawnTab_UserWidgetEditor(const F
 			];
 	}
 
+	const FText CustomLabel = IAudioWidgetInterface::Execute_GetEditorName(UserWidgets[WidgetIndex].Get());
+	if (!CustomLabel.IsEmpty())
+	{
+		Label = CustomLabel;
+	}
+
 	return SNew(SDockTab)
 		.Icon(IconBrush)
 		.Label(Label)
@@ -252,7 +288,7 @@ TSharedRef<SDockTab> FSoundEffectPresetEditor::SpawnTab_UserWidgetEditor(const F
 			.BorderImage(FEditorStyle::GetBrush("ToolPanel.GroupBorder"))
 			.Padding(0.0f)
 			[
-				UserWidget->TakeWidget()
+				UserWidgets[WidgetIndex]->TakeWidget()
 			]
 		];
 }
diff --git a/Engine/Source/Editor/AudioEditor/Private/Editors/SoundEffectPresetEditor.h b/Engine/Source/Editor/AudioEditor/Private/Editors/SoundEffectPresetEditor.h
--- a/Engine/Source/Editor/AudioEditor/Private/Editors/SoundEffectPresetEditor.h
+++ b/Engine/Source/Editor/AudioEditor/Private/Editors/SoundEffectPresetEditor.h
@@ -3,6 +3,7 @@
 
 #include "CoreMinimal.h"
 
+#include "AudioEditorSubsystem.h"
 #include "Blueprint/UserWidget.h"
 #include "EditorUndoClient.h"
 #include "Engine/DeveloperSettings.h"
@@ -19,6 +20,7 @@
 
 #include "SoundEffectPresetEditor.generated.h"
 
+
 // Forward Declarations
 class FCurveEditor;
 class IToolkitHost;
@@ -26,32 +28,31 @@ class SCurveEditorPanel;
 class SSoundEffectEditorPreviewViewport;
 class UClass;
 class UCurveBase;
-class UWidgetBlueprint;
 class USoundEffectPreset;
+class UWidgetBlueprint;
 
-
-UCLASS(Blueprintable)
-class USoundEffectPresetUserWidget : public UUserWidget
+UINTERFACE(Blueprintable)
+class AUDIOEDITOR_API USoundEffectPresetViewInterface : public UAudioWidgetInterface
 {
 	GENERATED_BODY()
+};
 
-public:
-	UPROPERTY(EditAnywhere, BlueprintReadWrite, Category = "General")
-	USoundEffectPreset* Preset;
-
-	UFUNCTION(BlueprintImplementableEvent)
-	void OnPresetChanged(FName PropertyName);
+class AUDIOEDITOR_API ISoundEffectPresetViewInterface : public IAudioWidgetInterface
+{
+	GENERATED_BODY()
 };
 
-class FSoundEffectPresetEditor : public FAssetEditorToolkit, public FNotifyHook, public FEditorUndoClient, public FGCObject
+
+class FSoundEffectPresetEditor : public FAssetEditorToolkit, public FNotifyHook, public FEditorUndoClient
 {
 public:
 	FSoundEffectPresetEditor();
 	virtual ~FSoundEffectPresetEditor() = default;
 
-	void Init(const EToolkitMode::Type Mode, const TSharedPtr<IToolkitHost>& InitToolkitHost, USoundEffectPreset* PresetToEdit);
+	void Init(const EToolkitMode::Type Mode, const TSharedPtr<IToolkitHost>& InitToolkitHost, USoundEffectPreset* PresetToEdit, const TArray<UUserWidget*>& InWidgetBlueprints);
 
 	/** FAssetEditorToolkit interface */
+	virtual FName GetEditorName() const override;
 	virtual FName GetToolkitFName() const override;
 	virtual FText GetBaseToolkitName() const override;
 	virtual FString GetWorldCentricTabPrefix() const override;
@@ -61,21 +62,17 @@ public:
 
 	/** FNotifyHook interface */
 	virtual void NotifyPostChange(const FPropertyChangedEvent& PropertyChangedEvent, FProperty* PropertyThatChanged) override;
-
-	/** FGCObject interface */
-	virtual void AddReferencedObjects(FReferenceCollector& Collector) override;
-	virtual FString GetReferencerName() const override;
-
+	virtual void NotifyPostChange(const FPropertyChangedEvent& PropertyChangedEvent, FEditPropertyChain* PropertyThatChanged) override;
 protected:
 	virtual void PostUndo(bool bSuccess) override;
 	virtual void PostRedo(bool bSuccess) override;
 
 private:
-	/** Finds and returns the UserWidget of the preset if one is provided by the use in the AudioEditor developer settings. */
-	void InitPresetWidget();
+	/** Initializes all preset user widgets. */
+	void InitPresetWidgets(const TArray<UUserWidget*>& InWidgets);
 
 	/**	Spawns the tab allowing for editing/viewing the blueprint widget for the associated SoundEffectPreset */
-	TSharedRef<SDockTab> SpawnTab_UserWidgetEditor(const FSpawnTabArgs& Args);
+	TSharedRef<SDockTab> SpawnTab_UserWidgetEditor(const FSpawnTabArgs& Args, int32 WidgetIndex);
 
 	/**	Spawns the tab allowing for editing/viewing details panel */
 	TSharedRef<SDockTab> SpawnTab_Properties(const FSpawnTabArgs& Args);
@@ -89,8 +86,8 @@ private:
 	/** Settings Editor App Identifier */
 	static const FName AppIdentifier;
 
-	USoundEffectPreset* SoundEffectPreset;
-	USoundEffectPresetUserWidget* UserWidget;
+	TStrongObjectPtr<USoundEffectPreset> SoundEffectPreset;
+	TArray<TStrongObjectPtr<UUserWidget>> UserWidgets;
 
 	/** Tab Ids */
 	static const FName PropertiesTabId;
diff --git a/Engine/Source/Editor/AudioEditor/Public/AudioEditorModule.h b/Engine/Source/Editor/AudioEditor/Public/AudioEditorModule.h
--- a/Engine/Source/Editor/AudioEditor/Public/AudioEditorModule.h
+++ b/Engine/Source/Editor/AudioEditor/Public/AudioEditorModule.h
@@ -66,15 +66,6 @@ public:
 	/** Returns the toolbar extensibility manager for the given audio editor type. */
 	virtual TSharedPtr<FExtensibilityManager> GetSoundCueToolBarExtensibilityManager() = 0;
 
-	/** Registers a custom widget blueprint with a SoundEffectPreset class for editing. */
-	virtual void RegisterSoundEffectPresetWidget(TSubclassOf<USoundEffectPreset> PresetClass, UWidgetBlueprint* WidgetBlueprint) = 0;
-
-	/** Returns custom widget blueprint for a given SoundEffectPreset class (or null if unset). */
-	virtual UWidgetBlueprint* GetSoundEffectPresetWidget(TSubclassOf<USoundEffectPreset> PresetClass) = 0;
-
-	/** Unregisters a custom widget blueprint with a SoundEffectPreset class for editing. */
-	virtual void UnregisterSoundEffectPresetWidget(TSubclassOf<USoundEffectPreset> PresetClass) = 0;
-
 	/** Replaces sound cue nodes in the graph. */
 	virtual void ReplaceSoundNodesInGraph(USoundCue* SoundCue, UDialogueWave* DialogueWave, TArray<USoundNode*>& NodesToReplace, const FDialogueContextMapping& ContextMapping) = 0;
 
diff --git a/Engine/Source/Editor/AudioEditor/Public/AudioEditorSubsystem.h b/Engine/Source/Editor/AudioEditor/Public/AudioEditorSubsystem.h
new file mode 100644
--- /dev/null
+++ b/Engine/Source/Editor/AudioEditor/Public/AudioEditorSubsystem.h
@@ -0,0 +1,64 @@
+// Copyright Epic Games, Inc. All Rights Reserved.
+#pragma once
+
+#include "AssetRegistry/AssetData.h"
+#include "Blueprint/UserWidget.h"
+#include "CoreMinimal.h"
+#include "EditorSubsystem.h"
+#include "Sound/SoundEffectPreset.h"
+#include "UObject/Class.h"
+#include "UObject/Interface.h"
+#include "UObject/ObjectMacros.h"
+#include "UObject/UObjectGlobals.h"
+
+#include "AudioEditorSubsystem.generated.h"
+
+
+// Forward Declarations
+class FText;
+class UUserWidget;
+
+
+UINTERFACE(Blueprintable)
+class AUDIOEDITOR_API UAudioWidgetInterface : public UInterface
+{
+	GENERATED_BODY()
+};
+
+class AUDIOEDITOR_API IAudioWidgetInterface
+{
+	GENERATED_BODY()
+
+public:
+	UFUNCTION(BlueprintImplementableEvent)
+	FText GetEditorName();
+
+	UFUNCTION(BlueprintImplementableEvent)
+	FName GetIconBrushName();
+
+	// Returns the class of object the compiled widget supports
+	UFUNCTION(BlueprintImplementableEvent)
+	UClass* GetClass();
+
+	UFUNCTION(BlueprintImplementableEvent)
+	void OnConstructed(UObject* Object);
+
+	UFUNCTION(BlueprintImplementableEvent)
+	void OnPropertyChanged(UObject* Object, FName PropertyName);
+
+};
+
+UCLASS()
+class AUDIOEDITOR_API UAudioEditorSubsystem : public UEditorSubsystem
+{
+	GENERATED_BODY()
+
+private:
+	static const TArray<FAssetData> GetWidgetBlueprintAssetData();
+	static bool ImplementsInterface(const FAssetData& InAssetData, UClass* InInterfaceClass);
+
+public:
+	// Returns user widgets that implement an AudioWidgetInterface.  Optionally, constructs only widgets that implement
+	// the provided AudioWidgetInterface type and/or widgets that support the given Object's parent class.
+	TArray<UUserWidget*> CreateUserWidgets(TSubclassOf<UAudioWidgetInterface> InWidgetClass, UClass* InObjectClass = nullptr) const;
+};
