Fix fallback icon string
#rnx

diff --git a/Engine/Source/Editor/AudioEditor/Private/Editors/SoundEffectPresetEditor.cpp b/Engine/Source/Editor/AudioEditor/Private/Editors/SoundEffectPresetEditor.cpp
--- a/Engine/Source/Editor/AudioEditor/Private/Editors/SoundEffectPresetEditor.cpp
+++ b/Engine/Source/Editor/AudioEditor/Private/Editors/SoundEffectPresetEditor.cpp
@@ -255,10 +255,10 @@ TSharedRef<SDockTab> FSoundEffectPresetEditor::SpawnTab_UserWidgetEditor(const F
 	FName IconBrushName = IAudioWidgetInterface::Execute_GetIconBrushName(UserWidgets[WidgetIndex].Get());
 	if (IconBrushName == FName())
 	{
-		IconBrushName = ("GenericEditor.Tabs.Properties");
+		IconBrushName = "GenericEditor.Tabs.Properties";
 	}
 
-	const FSlateBrush* IconBrush = FEditorStyle::GetBrush("SoundEffectPresetEditor.Tabs.Properties");
+	const FSlateBrush* IconBrush = FEditorStyle::GetBrush(IconBrushName);
 
 	FText Label = FText::FromString(SoundEffectPreset->GetName());
 	if (UserWidgets.Num() < WidgetIndex)
