Modifying AudioEditorSubsystems

Adding a new audio edtior settings category with some stubs for various audio editors.

Fixed a bug with ImplementsInterface failing to parse the string for a widget which implements only one interface.

diff --git a/Engine/Plugins/Runtime/Synthesis/Source/SynthesisEditor/Classes/SynthesisEditorSettings.h b/Engine/Plugins/Runtime/Synthesis/Source/SynthesisEditor/Classes/SynthesisEditorSettings.h
--- a/Engine/Plugins/Runtime/Synthesis/Source/SynthesisEditor/Classes/SynthesisEditorSettings.h
+++ b/Engine/Plugins/Runtime/Synthesis/Source/SynthesisEditor/Classes/SynthesisEditorSettings.h
@@ -2,13 +2,13 @@
 
 #pragma once
 
-#include "Engine/DeveloperSettings.h"
+#include "AudioEditorSettings.h"
 #include "WidgetBlueprint.h"
 
 #include "SynthesisEditorSettings.generated.h"
 
 
-UCLASS(config = EditorSettings, defaultconfig, meta = (DisplayName = "Synthesis"))
+UCLASS(config = EditorSettings, defaultconfig, meta = (DisplayName = "Synthesis and DSP Plugin"))
 class USynthesisEditorSettings : public UDeveloperSettings
 {
 	GENERATED_BODY()
diff --git a/Engine/Source/Editor/AudioEditor/AudioEditor.Build.cs b/Engine/Source/Editor/AudioEditor/AudioEditor.Build.cs
--- a/Engine/Source/Editor/AudioEditor/AudioEditor.Build.cs
+++ b/Engine/Source/Editor/AudioEditor/AudioEditor.Build.cs
@@ -25,6 +25,7 @@ public class AudioEditor : ModuleRules
 				"GameProjectGeneration",
 				"ToolMenus",
 				"UMG",
+				"DeveloperSettings",
 				"UMGEditor",
 				"AudioExtensions"
 			}
@@ -38,6 +39,7 @@ public class AudioEditor : ModuleRules
 				"CoreUObject",
 				"ApplicationCore",
 				"AudioMixer",
+				"SignalProcessing",
 				"InputCore",
 				"Engine",
 				"EditorFramework",
diff --git a/Engine/Source/Editor/AudioEditor/Classes/AudioEditorSettings.h b/Engine/Source/Editor/AudioEditor/Classes/AudioEditorSettings.h
new file mode 100644
--- /dev/null
+++ b/Engine/Source/Editor/AudioEditor/Classes/AudioEditorSettings.h
@@ -0,0 +1,59 @@
+// Copyright Epic Games, Inc. All Rights Reserved.
+
+#pragma once
+
+#include "Engine/DeveloperSettings.h"
+#include "WidgetBlueprint.h"
+
+#include "AudioEditorSettings.generated.h"
+
+
+UCLASS(config = EditorSettings, defaultconfig, meta = (DisplayName = "Submix Editor"))
+class AUDIOEDITOR_API UAudioSubmixEditorSettings : public UDeveloperSettings
+{
+	GENERATED_BODY()
+
+public:
+	// Which UMG widget to use to draw the submix graph node
+	UPROPERTY(config, EditAnywhere, Category = "Submix Graph", meta = (AllowedClasses = "WidgetBlueprint"))
+	FSoftObjectPath SubmixNodeWidget;
+
+	//~ Begin UDeveloperSettings
+	virtual FName GetCategoryName() const override { return TEXT("Audio"); }
+	//~ End UDeveloperSettings
+
+	//~ Begin UObject
+	virtual void PostEditChangeProperty(FPropertyChangedEvent& PropertyChangedEvent) override;
+	//~ End UObject
+};
+
+UCLASS(config = EditorSettings, defaultconfig, meta = (DisplayName = "Metasound Editor"))
+class AUDIOEDITOR_API UAudioMetasoundEditorSettings : public UDeveloperSettings
+{
+	GENERATED_BODY()
+
+public:
+	//~ Begin UDeveloperSettings
+	virtual FName GetCategoryName() const override { return TEXT("Audio"); }
+	//~ End UDeveloperSettings
+
+	//~ Begin UObject
+	virtual void PostEditChangeProperty(FPropertyChangedEvent& PropertyChangedEvent) override;
+	//~ End UObject
+};
+
+UCLASS(config = EditorSettings, defaultconfig, meta = (DisplayName = "SoundClass Editor"))
+class AUDIOEDITOR_API UAudioSoundClassEditorSettings : public UDeveloperSettings
+{
+	GENERATED_BODY()
+
+public:
+	//~ Begin UDeveloperSettings
+	virtual FName GetCategoryName() const override { return TEXT("Audio"); }
+	//~ End UDeveloperSettings
+
+	//~ Begin UObject
+	virtual void PostEditChangeProperty(FPropertyChangedEvent& PropertyChangedEvent) override;
+	//~ End UObject
+};
+
diff --git a/Engine/Source/Editor/AudioEditor/Private/AudioEditorSettings.cpp b/Engine/Source/Editor/AudioEditor/Private/AudioEditorSettings.cpp
new file mode 100644
--- /dev/null
+++ b/Engine/Source/Editor/AudioEditor/Private/AudioEditorSettings.cpp
@@ -0,0 +1,19 @@
+// Copyright Epic Games, Inc. All Rights Reserved.
+
+#include "AudioEditorSettings.h"
+
+void UAudioSubmixEditorSettings::PostEditChangeProperty(FPropertyChangedEvent& PropertyChangedEvent)
+{
+	Super::PostEditChangeProperty(PropertyChangedEvent);
+}
+
+
+void UAudioMetasoundEditorSettings::PostEditChangeProperty(FPropertyChangedEvent& PropertyChangedEvent)
+{
+	Super::PostEditChangeProperty(PropertyChangedEvent);
+}
+
+void UAudioSoundClassEditorSettings::PostEditChangeProperty(FPropertyChangedEvent& PropertyChangedEvent)
+{
+	Super::PostEditChangeProperty(PropertyChangedEvent);
+}
diff --git a/Engine/Source/Editor/AudioEditor/Private/AudioEditorSubsystem.cpp b/Engine/Source/Editor/AudioEditor/Private/AudioEditorSubsystem.cpp
--- a/Engine/Source/Editor/AudioEditor/Private/AudioEditorSubsystem.cpp
+++ b/Engine/Source/Editor/AudioEditor/Private/AudioEditorSubsystem.cpp
@@ -17,7 +17,7 @@
 
 #define LOCTEXT_NAMESPACE "AudioEditorSubsystem"
 
-TArray<UUserWidget*> UAudioEditorSubsystem::CreateUserWidgets(TSubclassOf<UAudioWidgetInterface> InWidgetClass, UClass* InObjectClass) const
+TArray<UUserWidget*> UAudioEditorSubsystem::CreateUserWidgets(TSubclassOf<UInterface> InWidgetClass, UClass* InObjectClass) const
 {
 	TArray<UUserWidget*> UserWidgets;
 	UClass* InterfaceClass = InWidgetClass ? InWidgetClass.Get() : UAudioWidgetInterface::StaticClass();
@@ -96,17 +96,20 @@ bool UAudioEditorSubsystem::ImplementsInterface(const FAssetData& InAssetData, U
 	const FString ImplementedInterfaces = InAssetData.GetTagValueRef<FString>(FBlueprintTags::ImplementedInterfaces);
 	if (!ImplementedInterfaces.IsEmpty())
 	{
-		FString FullInterface;
 		FString RemainingString;
 		FString InterfacePath;
 		FString CurrentString = *ImplementedInterfaces;
-		while (CurrentString.Split(TEXT(","), &FullInterface, &RemainingString))
+		FString FullInterface = CurrentString;
+		do
 		{
 			if (!CurrentString.StartsWith(TEXT("Graphs=(")))
 			{
 				if (FullInterface.Split(TEXT("\""), &CurrentString, &InterfacePath, ESearchCase::CaseSensitive))
 				{
-					InterfacePath.RemoveFromEnd(TEXT("\"'"));
+					if (!InterfacePath.RemoveFromEnd(TEXT("\"'")))
+					{
+						InterfacePath.RemoveFromEnd(TEXT("\"'))"));
+					}
 
 					const FCoreRedirectObjectName ResolvedInterfaceName = FCoreRedirects::GetRedirectedName(ECoreRedirectFlags::Type_Class, FCoreRedirectObjectName(InterfacePath));
 					if (InInterfaceClass->GetFName() == ResolvedInterfaceName.ObjectName)
@@ -117,7 +120,9 @@ bool UAudioEditorSubsystem::ImplementsInterface(const FAssetData& InAssetData, U
 			}
 
 			CurrentString = RemainingString;
-		}
+
+		} 
+		while (CurrentString.Split(TEXT(","), &FullInterface, &RemainingString));
 	}
 
 	return false;
diff --git a/Engine/Source/Editor/AudioEditor/Public/AudioEditorSubsystem.h b/Engine/Source/Editor/AudioEditor/Public/AudioEditorSubsystem.h
--- a/Engine/Source/Editor/AudioEditor/Public/AudioEditorSubsystem.h
+++ b/Engine/Source/Editor/AudioEditor/Public/AudioEditorSubsystem.h
@@ -58,7 +58,6 @@ private:
 	static bool ImplementsInterface(const FAssetData& InAssetData, UClass* InInterfaceClass);
 
 public:
-	// Returns user widgets that implement an AudioWidgetInterface.  Optionally, constructs only widgets that implement
-	// the provided AudioWidgetInterface type and/or widgets that support the given Object's parent class.
-	TArray<UUserWidget*> CreateUserWidgets(TSubclassOf<UAudioWidgetInterface> InWidgetClass, UClass* InObjectClass = nullptr) const;
+	// Returns user widgets that implement a widget interface of the given class and optionally widgets that support the given object's parent class.
+	TArray<UUserWidget*> CreateUserWidgets(TSubclassOf<UInterface> InWidgetClass, UClass* InObjectClass = nullptr) const;
 };
