Changes to sound submix graph node.

- Allowing UMG widgets to override the submix node.

diff --git a/Engine/Source/Editor/AudioEditor/Classes/SoundSubmixGraph/SoundSubmixGraphNode.h b/Engine/Source/Editor/AudioEditor/Classes/SoundSubmixGraph/SoundSubmixGraphNode.h
--- a/Engine/Source/Editor/AudioEditor/Classes/SoundSubmixGraph/SoundSubmixGraphNode.h
+++ b/Engine/Source/Editor/AudioEditor/Classes/SoundSubmixGraph/SoundSubmixGraphNode.h
@@ -3,12 +3,56 @@
 #pragma once
 
 #include "CoreMinimal.h"
+#include "Blueprint/UserWidget.h"
 #include "UObject/ObjectMacros.h"
 #include "EdGraph/EdGraphNode.h"
+#include "SGraphNode.h"
+#include "Sound/SoundSubmix.h"
+#include "AudioDevice.h"
+#include "DSP/EnvelopeFollower.h"
+#include "DSP/MultithreadedPatching.h"
+#include "AudioEditorSubsystem.h"
 #include "SoundSubmixGraphNode.generated.h"
 
 class USoundSubmixBase;
 
+UINTERFACE(Blueprintable)
+class AUDIOEDITOR_API USubmixNodeWidgetInterface : public UInterface
+{
+	GENERATED_BODY()
+};
+
+class AUDIOEDITOR_API ISubmixNodeWidgetInterface
+{
+	GENERATED_BODY()
+
+public:
+	UFUNCTION(BlueprintImplementableEvent)
+	void OnSubmixNodeConstructed(USoundSubmixBase* SoundSubmix);
+};
+
+class SSubmixGraphNode : public SGraphNode
+{
+public:
+	SLATE_BEGIN_ARGS(SSubmixGraphNode) {}
+
+	SLATE_ARGUMENT(TWeakObjectPtr<USoundSubmixBase>, SubmixBase)
+	SLATE_ARGUMENT(TWeakObjectPtr<UUserWidget>, SubmixNodeUserWidget)
+
+	SLATE_END_ARGS();
+
+	void Construct(const FArguments& InArgs, UEdGraphNode* InGraphNode);
+
+	virtual void UpdateGraphNode() override;
+	virtual void Tick(const FGeometry& AllottedGeometry, const double InCurrentTime, const float InDeltaTime) override;
+
+	TSharedRef<SWidget> CreateNodeContentArea() override;
+
+private:
+	TWeakObjectPtr<USoundSubmixBase> SubmixBase;
+	TWeakObjectPtr<UUserWidget> SubmixNodeUserWidget;
+};
+
 UCLASS(MinimalAPI)
 class USoundSubmixGraphNode : public UEdGraphNode
 {
@@ -18,6 +62,10 @@ class USoundSubmixGraphNode : public UEdGraphNode
 	UPROPERTY(VisibleAnywhere, instanced, Category=Sound)
 	USoundSubmixBase* SoundSubmix;
 
+	/** A user widget to use to represent the graph node */
+	UPROPERTY(Transient)
+	UUserWidget* SubmixNodeUserWidget;
+
 	/** Get the Pin that connects to all children */
 	UEdGraphPin* GetChildPin() const { return ChildPin; }
 
@@ -34,6 +82,7 @@ class USoundSubmixGraphNode : public UEdGraphNode
 	virtual bool CanCreateUnderSpecifiedSchema(const UEdGraphSchema* Schema) const override;
 	virtual bool CanUserDeleteNode() const override;
 	virtual FText GetNodeTitle(ENodeTitleType::Type TitleType) const override;
+	virtual TSharedPtr<SGraphNode> CreateVisualWidget();
 	//~ End UEdGraphNode Interface.
 
 private:
diff --git a/Engine/Source/Editor/AudioEditor/Private/SoundSubmixGraphNode.cpp b/Engine/Source/Editor/AudioEditor/Private/SoundSubmixGraphNode.cpp
--- a/Engine/Source/Editor/AudioEditor/Private/SoundSubmixGraphNode.cpp
+++ b/Engine/Source/Editor/AudioEditor/Private/SoundSubmixGraphNode.cpp
@@ -7,9 +7,91 @@
 #include "SoundSubmixEditor.h"
 #include "SoundSubmixDefaultColorPalette.h"
 #include "Subsystems/AssetEditorSubsystem.h"
+#include "AudioDeviceManager.h"
+#include "AudioEditorSettings.h"
+#include "Widgets/Input/SSlider.h"
+#include "AudioEditorSubsystem.h"
 
 #define LOCTEXT_NAMESPACE "SoundSubmixGraphNode"
 
+void SSubmixGraphNode::Construct(const FArguments& InArgs, UEdGraphNode* InGraphNode)
+{
+	SubmixBase = InArgs._SubmixBase;
+	SubmixNodeUserWidget = InArgs._SubmixNodeUserWidget;
+
+	if (SubmixNodeUserWidget.IsValid())
+	{
+		ISubmixNodeWidgetInterface::Execute_OnSubmixNodeConstructed(SubmixNodeUserWidget.Get(), SubmixBase.Get());
+	}
+	GraphNode = InGraphNode;
+	UpdateGraphNode();
+}
+
+void SSubmixGraphNode::UpdateGraphNode()
+{
+	SGraphNode::UpdateGraphNode();
+}
+
+void SSubmixGraphNode::Tick(const FGeometry& AllottedGeometry, const double InCurrentTime, const float InDeltaTime)
+{
+}
+
+TSharedRef<SWidget> SSubmixGraphNode::CreateNodeContentArea()
+{	
+	if (SubmixNodeUserWidget.IsValid())
+	{
+		// NODE CONTENT AREA
+		return SNew(SBorder)
+			.BorderImage(FEditorStyle::GetBrush("NoBorder"))
+			.HAlign(HAlign_Fill)
+			.VAlign(VAlign_Fill)
+			.Padding(FMargin(0, 3))
+			[
+				SNew(SVerticalBox)
+				+ SVerticalBox::Slot()
+				.AutoHeight()
+				[
+					SNew(SHorizontalBox)
+					+ SHorizontalBox::Slot()
+					.HAlign(HAlign_Left)
+					.FillWidth(1.0f)
+					[
+						// LEFT
+						SAssignNew(LeftNodeBox, SVerticalBox)
+					]
+					+ SHorizontalBox::Slot()
+					.AutoWidth()
+					.HAlign(HAlign_Right)
+					[
+						// RIGHT
+						SAssignNew(RightNodeBox, SVerticalBox)
+					]
+
+				]
+
+				+ SVerticalBox::Slot()
+				.AutoHeight()
+				[
+					SNew(SHorizontalBox)
+					+ SHorizontalBox::Slot()
+					.HAlign(HAlign_Left)
+					.FillWidth(1.0f)
+					[
+						SubmixNodeUserWidget->TakeWidget()
+					]
+				]
+			]
+		;
+	}
+	else
+	{
+		return SGraphNode::CreateNodeContentArea();
+	}
+
+}
+
+
+
 USoundSubmixGraphNode::USoundSubmixGraphNode(const FObjectInitializer& ObjectInitializer)
 	: Super(ObjectInitializer)
 	, ChildPin(NULL)
@@ -134,4 +216,30 @@ FText USoundSubmixGraphNode::GetNodeTitle(ENodeTitleType::Type TitleType) const
 		return Super::GetNodeTitle(TitleType);
 	}
 }
+
+TSharedPtr<SGraphNode> USoundSubmixGraphNode::CreateVisualWidget()
+{
+	if (USoundSubmixBase* SubmixBase = Cast<USoundSubmixBase>(SoundSubmix))
+	{
+		if (UAudioEditorSubsystem* AudioEditorSubsytem = GEditor->GetEditorSubsystem<UAudioEditorSubsystem>())
+		{
+			TArray<UUserWidget*> UserWidgets = AudioEditorSubsytem->CreateUserWidgets(USubmixNodeWidgetInterface::StaticClass());
+			if (!UserWidgets.IsEmpty())
+			{
+				SubmixNodeUserWidget = UserWidgets[0];
+			}
+		}
+
+		const UAudioSubmixEditorSettings* Settings = GetDefault<UAudioSubmixEditorSettings>();
+		check(Settings);
+
+		// Pass the owning submix and the user widgets to the graph node
+		return SNew(SSubmixGraphNode, this)
+			.SubmixBase(SubmixBase)
+			.SubmixNodeUserWidget(SubmixNodeUserWidget);
+	}
+	return nullptr;
+}
+
+
 #undef LOCTEXT_NAMESPACE
