Fix drag-and-drop misleadingly showing icon/tooltip saying you cannot drop Submixes on the Submix Graph Editor

diff --git a/Engine/Source/Editor/AudioEditor/Classes/SoundSubmixGraph/SoundSubmixGraphSchema.h b/Engine/Source/Editor/AudioEditor/Classes/SoundSubmixGraph/SoundSubmixGraphSchema.h
--- a/Engine/Source/Editor/AudioEditor/Classes/SoundSubmixGraph/SoundSubmixGraphSchema.h
+++ b/Engine/Source/Editor/AudioEditor/Classes/SoundSubmixGraph/SoundSubmixGraphSchema.h
@@ -118,6 +118,7 @@ class USoundSubmixGraphSchema : public UEdGraphSchema
 	virtual void BreakPinLinks(UEdGraphPin& TargetPin, bool bSendsNodeNotifcation) const override;
 	virtual void BreakSinglePinLink(UEdGraphPin* SourcePin, UEdGraphPin* TargetPin) const override;
 	virtual void DroppedAssetsOnGraph(const TArray<struct FAssetData>& Assets, const FVector2D& GraphPosition, UEdGraph* Graph) const override;
+	virtual void GetAssetsGraphHoverMessage(const TArray<FAssetData>& Assets, const UEdGraph* HoverGraph, FString& OutTooltipText, bool& OutOkIcon) const override;
 	//~ End EdGraphSchema Interface
 };
 
diff --git a/Engine/Source/Editor/AudioEditor/Private/SoundSubmixGraphSchema.cpp b/Engine/Source/Editor/AudioEditor/Private/SoundSubmixGraphSchema.cpp
--- a/Engine/Source/Editor/AudioEditor/Private/SoundSubmixGraphSchema.cpp
+++ b/Engine/Source/Editor/AudioEditor/Private/SoundSubmixGraphSchema.cpp
@@ -124,6 +124,45 @@ bool USoundSubmixGraphSchema::ConnectionCausesLoop(const UEdGraphPin* InputPin,
 	return OutputNode->SoundSubmix->RecurseCheckChild(InputNode->SoundSubmix);
 }
 
+void USoundSubmixGraphSchema::GetAssetsGraphHoverMessage(const TArray<FAssetData>& Assets, const UEdGraph* HoverGraph, FString& OutTooltipText, bool& OutOkIcon) const
+{
+	OutOkIcon = true;
+	OutTooltipText = TEXT("Add Submix to Graph.");
+
+	if (HoverGraph)
+	{
+		TArray<USoundSubmixGraphNode*> SubmixNodes;
+		HoverGraph->GetNodesOfClass<USoundSubmixGraphNode>(SubmixNodes);
+		for (USoundSubmixGraphNode* SubmixNode : SubmixNodes)
+		{
+			if (SubmixNode && SubmixNode->SoundSubmix)
+			{
+				auto MatchesSubmix = [&](const FAssetData& Asset)
+				{
+					return Asset.GetFullName() == SubmixNode->SoundSubmix->GetFullName();
+				};
+
+				if (Assets.ContainsByPredicate(MatchesSubmix))
+				{
+					OutOkIcon = false;
+					OutTooltipText = TEXT("Selected asset or assets already in graph.");
+					break;
+				}
+			}
+		}
+	}
+
+	for (const FAssetData& Data : Assets)
+	{
+		if (!Data.GetClass()->IsChildOf(USoundSubmixBase::StaticClass()))
+		{
+			OutOkIcon = false;
+			OutTooltipText = TEXT("Asset(s) must all be Submixes.");
+			break;
+		}
+	}
+}
+
 void USoundSubmixGraphSchema::GetGraphContextActions(FGraphContextMenuBuilder& ContextMenuBuilder) const
 {
 	const FText Name = LOCTEXT("NewSoundSubmix", "New Sound Submix");
