Actual fix for importing 24 bit audio

diff --git a/Engine/Source/Editor/AudioEditor/Private/Factories/SoundFactory.cpp b/Engine/Source/Editor/AudioEditor/Private/Factories/SoundFactory.cpp
--- a/Engine/Source/Editor/AudioEditor/Private/Factories/SoundFactory.cpp
+++ b/Engine/Source/Editor/AudioEditor/Private/Factories/SoundFactory.cpp
@@ -142,29 +142,6 @@ UObject* USoundFactory::FactoryCreateBinary
 		SoundObject = CreateObject(Class, InParent, Name, Flags, Context, FileType, Buffer, BufferEnd, Warn);
 	}
 
-	// If we do not, we can use LibSoundFile here to attempt to convert the file to a 16 bit wave file.
-#if WITH_SNDFILE_IO
-	if (!SoundObject)
-	{
-		// Read raw audio data
-		TArray<uint8> RawAudioData;
-		RawAudioData.Empty(BufferEnd - Buffer);
-		RawAudioData.AddUninitialized(BufferEnd - Buffer);
-		FMemory::Memcpy(RawAudioData.GetData(), Buffer, RawAudioData.Num());
-
-		// Convert audio data to a wav file in memory
-		TArray<uint8> RawWaveData;
-		if (Audio::ConvertAudioToWav(RawAudioData, RawWaveData))
-		{
-			const uint8* Ptr = &RawWaveData[0];
-
-			// Perpetuate the setting of the suppression flag to avoid
-			// user notification if we attempt to call CreateObject twice
-			SoundObject = CreateObject(Class, InParent, Name, Flags, Context, TEXT("WAV"), Ptr, Ptr + RawWaveData.Num(), Warn);
-		}
-	}
-#endif // WITH_SNDFILE_IO
-
 	if (!SoundObject)
 	{
 		// Unrecognized sound format
@@ -303,43 +280,62 @@ UObject* USoundFactory::CreateObject
 		bool bIsFuMa = (FuMaTag == TEXT("_fuma"));
 
 		TArray<uint8> RawWaveData;
-		RawWaveData.Empty(BufferEnd - Buffer);
-		RawWaveData.AddUninitialized(BufferEnd - Buffer);
+		uint32 RawWaveDataBufferSize = BufferEnd - Buffer;
+		RawWaveData.Empty(RawWaveDataBufferSize);
+		RawWaveData.AddUninitialized(RawWaveDataBufferSize);
 		FMemory::Memcpy(RawWaveData.GetData(), Buffer, RawWaveData.Num());
 
+		// Converted buffer if we need it
+		TArray<uint8> ConvertedRawWaveData;
+
 		// Read the wave info and make sure we have valid wave data
 		FWaveModInfo WaveInfo;
 		FString ErrorMessage;
-		if (WaveInfo.ReadWaveInfo(RawWaveData.GetData(), RawWaveData.Num(), &ErrorMessage))
+		if (!WaveInfo.ReadWaveInfo(RawWaveData.GetData(), RawWaveData.Num(), &ErrorMessage))
 		{
-			// Validate if somebody has used the ambiX or FuMa tag that the ChannelCount is 4 channels
-			if ((bIsAmbiX || bIsFuMa) && (int32)*WaveInfo.pChannels != 4)
-			{
-				Warn->Logf(ELogVerbosity::Error, TEXT("Tried to import ambisonics format file but requires exactly 4 channels: '%s'"), *Name.ToString());
-				GEditor->GetEditorSubsystem<UImportSubsystem>()->BroadcastAssetPostImport(this, nullptr);
-				return nullptr;
-			}
+			Warn->Logf(ELogVerbosity::Error, TEXT("Unable to read wave file '%s' - \"%s\""), *Name.ToString(), *ErrorMessage);
+			GEditor->GetEditorSubsystem<UImportSubsystem>()->BroadcastAssetPostImport(this, nullptr);
+			return nullptr;
+		}
 
-			// If we are not using libSoundFile, we cannot support non-16 bit WAV files.
-			if (*WaveInfo.pBitsPerSample != 16)
+		if (*WaveInfo.pBitsPerSample != 16)
+		{
+#if WITH_SNDFILE_IO
+			// Attempt to convert to 16 bit audio
+			if (Audio::ConvertAudioToWav(RawWaveData, ConvertedRawWaveData))
 			{
-#if !WITH_SNDFILE_IO
-				WaveInfo.ReportImportFailure();
-				Warn->Logf(ELogVerbosity::Error, TEXT("Only 16 bit WAV source files are supported (%s) on this editor platform."), *Name.ToString());
-				GEditor->GetEditorSubsystem<UImportSubsystem>()->BroadcastAssetPostImport(this, nullptr);
-				return nullptr;
-#endif // WITH_SNDFILE_IO
-				
+				WaveInfo = FWaveModInfo();
+				if (!WaveInfo.ReadWaveInfo(ConvertedRawWaveData.GetData(), ConvertedRawWaveData.Num(), &ErrorMessage))
+				{
+					Warn->Logf(ELogVerbosity::Error, TEXT("Failed to convert to 16 bit WAV source on import."));
+					GEditor->GetEditorSubsystem<UImportSubsystem>()->BroadcastAssetPostImport(this, nullptr);
+					return nullptr;
+				}
 			}
+
+			// Copy over the data
+			Buffer = ConvertedRawWaveData.GetData();
+			RawWaveDataBufferSize = ConvertedRawWaveData.Num() * sizeof(uint8);
+
+#else
+			WaveInfo.ReportImportFailure();
+			Warn->Logf(ELogVerbosity::Error, TEXT("Only 16 bit WAV source files are supported on this editor platform."));
+			GEditor->GetEditorSubsystem<UImportSubsystem>()->BroadcastAssetPostImport(this, nullptr);
+			return nullptr;
+#endif
 		}
-		else
+
+		// Make sure we have 16 bit audio at this point
+		check(*WaveInfo.pBitsPerSample == 16);
+
+		// Validate if somebody has used the ambiX or FuMa tag that the ChannelCount is 4 channels
+		if ((bIsAmbiX || bIsFuMa) && (int32)*WaveInfo.pChannels != 4)
 		{
-			Warn->Logf(ELogVerbosity::Error, TEXT("Unable to read wave file '%s' - \"%s\""), *Name.ToString(), *ErrorMessage);
+			Warn->Logf(ELogVerbosity::Error, TEXT("Tried to import ambisonics format file but requires exactly 4 channels: '%s'"), *Name.ToString());
 			GEditor->GetEditorSubsystem<UImportSubsystem>()->BroadcastAssetPostImport(this, nullptr);
 			return nullptr;
 		}
 
-
 		// Use pre-existing sound if it exists and we want to keep settings,
 		// otherwise create new sound and import raw data.
 		USoundWave* Sound = (bUseExistingSettings && ExistingSound) ? ExistingSound : NewObject<USoundWave>(InParent, Name, Flags, TemplateSoundWave.Get());
@@ -515,8 +511,8 @@ UObject* USoundFactory::CreateObject
 		{
 			// For mono and stereo assets, just copy the data into the buffer
 			Sound->RawData.Lock(LOCK_READ_WRITE);
-			void* LockedData = Sound->RawData.Realloc(BufferEnd - Buffer);
-			FMemory::Memcpy(LockedData, Buffer, BufferEnd - Buffer);
+			void* LockedData = Sound->RawData.Realloc(RawWaveDataBufferSize);
+			FMemory::Memcpy(LockedData, Buffer, RawWaveDataBufferSize);
 			Sound->RawData.Unlock();
 		}
 
