Fixing up sound class editor tooltip

diff --git a/Engine/Source/Editor/AudioEditor/Private/SoundClassEditorCommands.cpp b/Engine/Source/Editor/AudioEditor/Private/SoundClassEditorCommands.cpp
--- a/Engine/Source/Editor/AudioEditor/Private/SoundClassEditorCommands.cpp
+++ b/Engine/Source/Editor/AudioEditor/Private/SoundClassEditorCommands.cpp
@@ -6,8 +6,8 @@
 
 void FSoundClassEditorCommands::RegisterCommands()
 {
-	UI_COMMAND(ToggleSolo, "Solo", "Toggles Soloing this cue", EUserInterfaceActionType::ToggleButton, FInputChord(EKeys::S));
-	UI_COMMAND(ToggleMute, "Mute", "Toggles Muting this cue", EUserInterfaceActionType::ToggleButton, FInputChord(EKeys::M));
+	UI_COMMAND(ToggleSolo, "Solo", "Toggles Soloing this sound class", EUserInterfaceActionType::ToggleButton, FInputChord(EKeys::S));
+	UI_COMMAND(ToggleMute, "Mute", "Toggles Muting this sound class", EUserInterfaceActionType::ToggleButton, FInputChord(EKeys::M));
 }
 
 #undef LOCTEXT_NAMESPACE
