Add Audio Editor Settings options to pin various asset types
- Move MetaSounds editor settings to Content Editor Category, remove unused settings

diff --git a/Engine/Plugins/Runtime/Metasound/Source/MetasoundEditor/MetasoundEditor.Build.cs b/Engine/Plugins/Runtime/Metasound/Source/MetasoundEditor/MetasoundEditor.Build.cs
--- a/Engine/Plugins/Runtime/Metasound/Source/MetasoundEditor/MetasoundEditor.Build.cs
+++ b/Engine/Plugins/Runtime/Metasound/Source/MetasoundEditor/MetasoundEditor.Build.cs
@@ -13,6 +13,7 @@ namespace UnrealBuildTool.Rules
 			(
 				new string[]
 				{
+					"AudioEditor",
 					"AudioExtensions",
 					"AudioMixer",
 					"KismetWidgets",
diff --git a/Engine/Plugins/Runtime/Metasound/Source/MetasoundEditor/Private/MetasoundAssetTypeActions.cpp b/Engine/Plugins/Runtime/Metasound/Source/MetasoundEditor/Private/MetasoundAssetTypeActions.cpp
--- a/Engine/Plugins/Runtime/Metasound/Source/MetasoundEditor/Private/MetasoundAssetTypeActions.cpp
+++ b/Engine/Plugins/Runtime/Metasound/Source/MetasoundEditor/Private/MetasoundAssetTypeActions.cpp
@@ -4,6 +4,7 @@
 
 #include "ContentBrowserModule.h"
 #include "ContentBrowserMenuContexts.h"
+#include "Editor/EditorPerProjectUserSettings.h"
 #include "IContentBrowserSingleton.h"
 #include "ObjectEditorUtils.h"
 #include "Metasound.h"
@@ -11,6 +12,7 @@
 #include "MetasoundSource.h"
 #include "MetasoundEditor.h"
 #include "MetasoundEditorGraphBuilder.h"
+#include "MetasoundEditorSettings.h"
 #include "MetasoundFactory.h"
 #include "MetasoundUObjectRegistry.h"
 #include "ToolMenus.h"
@@ -133,6 +135,12 @@ namespace Metasound
 
 		const TArray<FText>& FAssetTypeActions_MetaSound::GetSubMenus() const
 		{
+			if (GetDefault<UMetasoundEditorSettings>()->bPinMetaSoundInAssetMenu)
+			{
+				static const TArray<FText> SubMenus;
+				return SubMenus;
+			}
+			
 			static const TArray<FText> SubMenus
 			{
 				LOCTEXT("AssetSoundMetaSoundsSubMenu", "MetaSounds"),
@@ -172,6 +180,22 @@ namespace Metasound
 			// referenced graphs "interface."
 // 			RegisterPresetAction<UMetaSound, UMetaSoundFactory, UMetaSoundSource>(PresetLabel, PresetToolTip);
 		}
+
+		const TArray<FText>& FAssetTypeActions_MetaSoundSource::GetSubMenus() const
+		{
+			if (GetDefault<UMetasoundEditorSettings>()->bPinMetaSoundSourceInAssetMenu)
+			{
+				static const TArray<FText> SubMenus;
+				return SubMenus;
+			}
+
+			static const TArray<FText> SubMenus
+			{
+				LOCTEXT("AssetSoundMetaSoundSourceSubMenu", "MetaSounds"),
+			};
+
+			return SubMenus;
+		}
 	} // namespace Editor
 } // namespace Metasound
 #undef LOCTEXT_NAMESPACE //MetaSoundEditor
diff --git a/Engine/Plugins/Runtime/Metasound/Source/MetasoundEditor/Private/MetasoundAssetTypeActions.h b/Engine/Plugins/Runtime/Metasound/Source/MetasoundEditor/Private/MetasoundAssetTypeActions.h
--- a/Engine/Plugins/Runtime/Metasound/Source/MetasoundEditor/Private/MetasoundAssetTypeActions.h
+++ b/Engine/Plugins/Runtime/Metasound/Source/MetasoundEditor/Private/MetasoundAssetTypeActions.h
@@ -34,6 +34,7 @@ namespace Metasound
 			virtual FColor GetTypeColor() const override { return FColor(103, 214, 66); }
 			virtual UClass* GetSupportedClass() const override;
 			virtual uint32 GetCategories() override { return EAssetTypeCategories::Sounds; }
+			virtual const TArray<FText>& GetSubMenus() const override;
 			virtual void OpenAssetEditor(const TArray<UObject*>& InObjects, TSharedPtr<IToolkitHost> ToolkitHost) override;
 
 			static void RegisterMenuActions();
diff --git a/Engine/Plugins/Runtime/Metasound/Source/MetasoundEditor/Private/MetasoundEditorModule.cpp b/Engine/Plugins/Runtime/Metasound/Source/MetasoundEditor/Private/MetasoundEditorModule.cpp
--- a/Engine/Plugins/Runtime/Metasound/Source/MetasoundEditor/Private/MetasoundEditorModule.cpp
+++ b/Engine/Plugins/Runtime/Metasound/Source/MetasoundEditor/Private/MetasoundEditorModule.cpp
@@ -681,7 +681,7 @@ namespace Metasound
 
 				ISettingsModule& SettingsModule = FModuleManager::LoadModuleChecked<ISettingsModule>("Settings");
 
-				SettingsModule.RegisterSettings("Editor", "Audio", "MetaSound Editor",
+				SettingsModule.RegisterSettings("Editor", "ContentEditors", "MetaSound Editor",
 					NSLOCTEXT("MetaSoundsEditor", "MetaSoundEditorSettingsName", "MetaSound Editor"),
 					NSLOCTEXT("MetaSoundsEditor", "MetaSoundEditorSettingsDescription", "Customize MetaSound Editor."),
 					GetMutableDefault<UMetasoundEditorSettings>()
diff --git a/Engine/Plugins/Runtime/Metasound/Source/MetasoundEditor/Private/MetasoundEditorSettings.h b/Engine/Plugins/Runtime/Metasound/Source/MetasoundEditor/Private/MetasoundEditorSettings.h
--- a/Engine/Plugins/Runtime/Metasound/Source/MetasoundEditor/Private/MetasoundEditorSettings.h
+++ b/Engine/Plugins/Runtime/Metasound/Source/MetasoundEditor/Private/MetasoundEditorSettings.h
@@ -22,6 +22,14 @@ class METASOUNDEDITOR_API UMetasoundEditorSettings : public UObject
 	GENERATED_UCLASS_BODY()
 
 public:
+	/** Whether to pin the MetaSound asset type when creating new assets. */
+	UPROPERTY(EditAnywhere, config, DisplayName = "Pin MetaSound in Asset Menu", Category = AssetMenu)
+	bool bPinMetaSoundInAssetMenu = false;
+
+	/** Whether to pin the MetaSound Source asset type when creating new assets. */
+	UPROPERTY(EditAnywhere, config, DisplayName = "Pin MetaSound Source in Asset Menu", Category = AssetMenu)
+	bool bPinMetaSoundSourceInAssetMenu = true;
+
 	/** Default author title to use when authoring a new
 	  * MetaSound.  If empty, uses machine name by default.
 	  */
diff --git a/Engine/Plugins/Runtime/SoundCueTemplates/Source/SoundCueTemplatesEditor/Private/AssetTypeActions_SoundCueTemplate.cpp b/Engine/Plugins/Runtime/SoundCueTemplates/Source/SoundCueTemplatesEditor/Private/AssetTypeActions_SoundCueTemplate.cpp
--- a/Engine/Plugins/Runtime/SoundCueTemplates/Source/SoundCueTemplatesEditor/Private/AssetTypeActions_SoundCueTemplate.cpp
+++ b/Engine/Plugins/Runtime/SoundCueTemplates/Source/SoundCueTemplatesEditor/Private/AssetTypeActions_SoundCueTemplate.cpp
@@ -5,6 +5,7 @@
 #include "SoundCueTemplate.h"
 #include "SoundCueTemplateFactory.h"
 
+#include "AudioEditorSettings.h"
 #include "Components/AudioComponent.h"
 #include "ToolMenus.h"
 #include "ToolMenuSection.h"
@@ -62,19 +63,17 @@ void FAssetTypeActions_SoundCueTemplate::ExecuteCopyToSoundCue(TArray<TWeakObjec
 
 const TArray<FText>& FAssetTypeActions_SoundCueTemplate::GetSubMenus() const
 {
-	if (FModuleManager::Get().IsModuleLoaded(TEXT("MetasoundEditor")))
-	{
-		static const TArray<FText> AssetTypeActionSubMenu
-		{
-			LOCTEXT("AssetSoundCueSubMenu", "Legacy")
-		};
-		return AssetTypeActionSubMenu;
-	}
-	else
+	if (GetDefault<UAudioEditorSettings>()->bPinSoundCueTemplateInAssetMenu)
 	{
 		static const TArray<FText> AssetTypeActionSubMenu;
 		return AssetTypeActionSubMenu;
 	}
+
+	static const TArray<FText> AssetTypeActionSubMenu
+	{
+		LOCTEXT("AssetSoundCueSubMenu", "Source")
+	};
+	return AssetTypeActionSubMenu;
 }
 
 void FAssetActionExtender_SoundCueTemplate::RegisterMenus()
diff --git a/Engine/Source/Editor/AudioEditor/Classes/AudioEditorSettings.h b/Engine/Source/Editor/AudioEditor/Classes/AudioEditorSettings.h
--- a/Engine/Source/Editor/AudioEditor/Classes/AudioEditorSettings.h
+++ b/Engine/Source/Editor/AudioEditor/Classes/AudioEditorSettings.h
@@ -3,39 +3,33 @@
 #pragma once
 
 #include "Engine/DeveloperSettings.h"
-#include "WidgetBlueprint.h"
 
 #include "AudioEditorSettings.generated.h"
 
-
-UCLASS(config = EditorSettings, defaultconfig, meta = (DisplayName = "Submix Editor"))
-class AUDIOEDITOR_API UAudioSubmixEditorSettings : public UDeveloperSettings
+UCLASS(config = EditorPerProjectUserSettings, defaultconfig, meta = (DisplayName = "Audio"))
+class AUDIOEDITOR_API UAudioEditorSettings : public UDeveloperSettings
 {
 	GENERATED_BODY()
 
 public:
-	// Which UMG widget to use to draw the submix graph node
-	UPROPERTY(config, EditAnywhere, Category = "Submix Graph", meta = (AllowedClasses = "WidgetBlueprint"))
-	FSoftObjectPath SubmixNodeWidget;
+	/** Whether to pin the Sound Cue asset type when creating new assets. */
+	UPROPERTY(EditAnywhere, config, Category = AssetMenu)
+	bool bPinSoundCueInAssetMenu = false;
 
-	//~ Begin UDeveloperSettings
-	virtual FName GetCategoryName() const override { return TEXT("Audio"); }
-	//~ End UDeveloperSettings
+	/** Whether to pin the Sound Cue Template asset type when creating new assets. */
+	UPROPERTY(EditAnywhere, config, Category = AssetMenu)
+	bool bPinSoundCueTemplateInAssetMenu = false;
 
-	//~ Begin UObject
-	virtual void PostEditChangeProperty(FPropertyChangedEvent& PropertyChangedEvent) override;
-	//~ End UObject
-};
+	/** Whether to pin the Sound Attenuation asset type when creating new assets. */
+	UPROPERTY(EditAnywhere, config, Category = AssetMenu)
+	bool bPinSoundAttenuationInAssetMenu = true;
 
+	/** Whether to pin the Sound Concurrency asset type when creating new assets. */
+	UPROPERTY(EditAnywhere, config, Category = AssetMenu)
+	bool bPinSoundConcurrencyInAssetMenu = true;
 
-UCLASS(config = EditorSettings, defaultconfig, meta = (DisplayName = "SoundClass Editor"))
-class AUDIOEDITOR_API UAudioSoundClassEditorSettings : public UDeveloperSettings
-{
-	GENERATED_BODY()
-
-public:
 	//~ Begin UDeveloperSettings
-	virtual FName GetCategoryName() const override { return TEXT("Audio"); }
+	virtual FName GetCategoryName() const override { return TEXT("General"); }
 	//~ End UDeveloperSettings
 
 	//~ Begin UObject
diff --git a/Engine/Source/Editor/AudioEditor/Private/AssetTypeActions/AssetTypeActions_SoundAttenuation.cpp b/Engine/Source/Editor/AudioEditor/Private/AssetTypeActions/AssetTypeActions_SoundAttenuation.cpp
--- a/Engine/Source/Editor/AudioEditor/Private/AssetTypeActions/AssetTypeActions_SoundAttenuation.cpp
+++ b/Engine/Source/Editor/AudioEditor/Private/AssetTypeActions/AssetTypeActions_SoundAttenuation.cpp
@@ -1,9 +1,28 @@
 // Copyright Epic Games, Inc. All Rights Reserved.
 
 #include "AssetTypeActions/AssetTypeActions_SoundAttenuation.h"
+#include "AudioEditorSettings.h"
 #include "Sound/SoundAttenuation.h"
 
+#define LOCTEXT_NAMESPACE "AssetTypeActions"
+
 UClass* FAssetTypeActions_SoundAttenuation::GetSupportedClass() const
 {
 	return USoundAttenuation::StaticClass();
 }
+
+const TArray<FText>& FAssetTypeActions_SoundAttenuation::GetSubMenus() const
+{
+	if (GetDefault<UAudioEditorSettings>()->bPinSoundAttenuationInAssetMenu)
+	{
+		static const TArray<FText> AssetTypeActionSubMenu;
+		return AssetTypeActionSubMenu;
+	}
+	static const TArray<FText> AssetTypeActionSubMenu
+	{
+		LOCTEXT("AssetSoundAttenuationSubMenu", "Spatialization"),
+	};
+	return AssetTypeActionSubMenu;
+}
+
+#undef LOCTEXT_NAMESPACE
diff --git a/Engine/Source/Editor/AudioEditor/Private/AssetTypeActions/AssetTypeActions_SoundAttenuation.h b/Engine/Source/Editor/AudioEditor/Private/AssetTypeActions/AssetTypeActions_SoundAttenuation.h
--- a/Engine/Source/Editor/AudioEditor/Private/AssetTypeActions/AssetTypeActions_SoundAttenuation.h
+++ b/Engine/Source/Editor/AudioEditor/Private/AssetTypeActions/AssetTypeActions_SoundAttenuation.h
@@ -13,4 +13,5 @@ public:
 	virtual FColor GetTypeColor() const override { return FColor(77, 120, 239); }
 	virtual UClass* GetSupportedClass() const override;
 	virtual uint32 GetCategories() override { return EAssetTypeCategories::Sounds; }
+	virtual const TArray<FText>& GetSubMenus() const override;
 };
diff --git a/Engine/Source/Editor/AudioEditor/Private/AssetTypeActions/AssetTypeActions_SoundConcurrency.cpp b/Engine/Source/Editor/AudioEditor/Private/AssetTypeActions/AssetTypeActions_SoundConcurrency.cpp
--- a/Engine/Source/Editor/AudioEditor/Private/AssetTypeActions/AssetTypeActions_SoundConcurrency.cpp
+++ b/Engine/Source/Editor/AudioEditor/Private/AssetTypeActions/AssetTypeActions_SoundConcurrency.cpp
@@ -1,9 +1,28 @@
 // Copyright Epic Games, Inc. All Rights Reserved.
 
 #include "AssetTypeActions/AssetTypeActions_SoundConcurrency.h"
+#include "AudioEditorSettings.h"
 #include "Sound/SoundConcurrency.h"
 
+#define LOCTEXT_NAMESPACE "AssetTypeActions"
+
 UClass* FAssetTypeActions_SoundConcurrency::GetSupportedClass() const
 {
 	return USoundConcurrency::StaticClass();
 }
+
+const TArray<FText>& FAssetTypeActions_SoundConcurrency::GetSubMenus() const 
+{
+	if (GetDefault<UAudioEditorSettings>()->bPinSoundConcurrencyInAssetMenu)
+	{
+		static const TArray<FText> AssetTypeActionSubMenu;
+		return AssetTypeActionSubMenu;
+	}
+	static const TArray<FText> AssetTypeActionSubMenu
+	{
+		LOCTEXT("AssetSoundConcurrencySubMenu", "Mix"),
+	};
+	return AssetTypeActionSubMenu;
+}
+
+#undef LOCTEXT_NAMESPACE
diff --git a/Engine/Source/Editor/AudioEditor/Private/AssetTypeActions/AssetTypeActions_SoundConcurrency.h b/Engine/Source/Editor/AudioEditor/Private/AssetTypeActions/AssetTypeActions_SoundConcurrency.h
--- a/Engine/Source/Editor/AudioEditor/Private/AssetTypeActions/AssetTypeActions_SoundConcurrency.h
+++ b/Engine/Source/Editor/AudioEditor/Private/AssetTypeActions/AssetTypeActions_SoundConcurrency.h
@@ -13,4 +13,5 @@ public:
 	virtual FColor GetTypeColor() const override { return FColor(77, 100, 139); }
 	virtual UClass* GetSupportedClass() const override;
 	virtual uint32 GetCategories() override { return EAssetTypeCategories::Sounds; }
+	virtual const TArray<FText>& GetSubMenus() const override;
 };
diff --git a/Engine/Source/Editor/AudioEditor/Private/AssetTypeActions/AssetTypeActions_SoundCue.cpp b/Engine/Source/Editor/AudioEditor/Private/AssetTypeActions/AssetTypeActions_SoundCue.cpp
--- a/Engine/Source/Editor/AudioEditor/Private/AssetTypeActions/AssetTypeActions_SoundCue.cpp
+++ b/Engine/Source/Editor/AudioEditor/Private/AssetTypeActions/AssetTypeActions_SoundCue.cpp
@@ -1,6 +1,7 @@
 // Copyright Epic Games, Inc. All Rights Reserved.
 
 #include "AssetTypeActions/AssetTypeActions_SoundCue.h"
+#include "AudioEditorSettings.h"
 #include "Sound/SoundAttenuation.h"
 #include "Misc/PackageName.h"
 #include "Sound/SoundCue.h"
@@ -123,19 +124,16 @@ void FAssetTypeActions_SoundCue::ExecuteConsolidateAttenuation(TArray<TWeakObjec
 
 const TArray<FText>& FAssetTypeActions_SoundCue::GetSubMenus() const
 {
-	if (FModuleManager::Get().IsModuleLoaded(TEXT("MetasoundEditor")))
-	{
-		static const TArray<FText> AssetTypeActionSubMenu
-		{
-			LOCTEXT("AssetSoudnCueSubMenu", "Legacy")
-		};
-		return AssetTypeActionSubMenu;
-	}
-	else
+	if (GetDefault<UAudioEditorSettings>()->bPinSoundCueInAssetMenu)
 	{
 		static const TArray<FText> AssetTypeActionSubMenu;
 		return AssetTypeActionSubMenu;
 	}
+	static const TArray<FText> AssetTypeActionSubMenu
+	{
+		LOCTEXT("AssetSoundCueSubMenu", "Source")
+	};
+	return AssetTypeActionSubMenu;
 }
 
 #undef LOCTEXT_NAMESPACE
diff --git a/Engine/Source/Editor/AudioEditor/Private/AudioEditorSettings.cpp b/Engine/Source/Editor/AudioEditor/Private/AudioEditorSettings.cpp
--- a/Engine/Source/Editor/AudioEditor/Private/AudioEditorSettings.cpp
+++ b/Engine/Source/Editor/AudioEditor/Private/AudioEditorSettings.cpp
@@ -2,12 +2,7 @@
 
 #include "AudioEditorSettings.h"
 
-void UAudioSubmixEditorSettings::PostEditChangeProperty(FPropertyChangedEvent& PropertyChangedEvent)
-{
-	Super::PostEditChangeProperty(PropertyChangedEvent);
-}
-
-void UAudioSoundClassEditorSettings::PostEditChangeProperty(FPropertyChangedEvent& PropertyChangedEvent)
+void UAudioEditorSettings::PostEditChangeProperty(FPropertyChangedEvent& PropertyChangedEvent)
 {
 	Super::PostEditChangeProperty(PropertyChangedEvent);
 }
diff --git a/Engine/Source/Editor/AudioEditor/Private/SoundSubmixGraphNode.cpp b/Engine/Source/Editor/AudioEditor/Private/SoundSubmixGraphNode.cpp
--- a/Engine/Source/Editor/AudioEditor/Private/SoundSubmixGraphNode.cpp
+++ b/Engine/Source/Editor/AudioEditor/Private/SoundSubmixGraphNode.cpp
@@ -3,7 +3,6 @@
 #include "SoundSubmixGraph/SoundSubmixGraphNode.h"
 
 #include "AudioDeviceManager.h"
-#include "AudioEditorSettings.h"
 #include "Audio/AudioWidgetSubsystem.h"
 #include "Sound/SoundSubmix.h"
 #include "SoundSubmixDefaultColorPalette.h"
@@ -228,9 +227,6 @@ TSharedPtr<SGraphNode> USoundSubmixGraphNode::CreateVisualWidget()
 				}
 			}
 
-			const UAudioSubmixEditorSettings* Settings = GetDefault<UAudioSubmixEditorSettings>();
-			check(Settings);
-
 			// Pass the owning submix and the user widgets to the graph node
 			return SNew(SSubmixGraphNode, this)
 				.SubmixBase(SubmixBase)
