Pin Sound Cue asset type by default when creating new assets

diff --git a/Engine/Plugins/Runtime/Metasound/Source/MetasoundEditor/MetasoundEditor.Build.cs b/Engine/Plugins/Runtime/Metasound/Source/MetasoundEditor/MetasoundEditor.Build.cs
--- a/Engine/Plugins/Runtime/Metasound/Source/MetasoundEditor/MetasoundEditor.Build.cs
+++ b/Engine/Plugins/Runtime/Metasound/Source/MetasoundEditor/MetasoundEditor.Build.cs
@@ -13,7 +13,6 @@ namespace UnrealBuildTool.Rules
 			(
 				new string[]
 				{
-					"AudioEditor",
 					"AudioExtensions",
 					"AudioMixer",
 					"KismetWidgets",
diff --git a/Engine/Source/Editor/AudioEditor/Classes/AudioEditorSettings.h b/Engine/Source/Editor/AudioEditor/Classes/AudioEditorSettings.h
--- a/Engine/Source/Editor/AudioEditor/Classes/AudioEditorSettings.h
+++ b/Engine/Source/Editor/AudioEditor/Classes/AudioEditorSettings.h
@@ -14,7 +14,7 @@ class AUDIOEDITOR_API UAudioEditorSettings : public UDeveloperSettings
 public:
 	/** Whether to pin the Sound Cue asset type when creating new assets. */
 	UPROPERTY(EditAnywhere, config, Category = AssetMenu)
-	bool bPinSoundCueInAssetMenu = false;
+	bool bPinSoundCueInAssetMenu = true;
 
 	/** Whether to pin the Sound Cue Template asset type when creating new assets. */
 	UPROPERTY(EditAnywhere, config, Category = AssetMenu)
