Fix clang-13 crash in AudioEditor and AudioSynesthesia
#rnx
PR #8441

diff --git a/Engine/Plugins/Runtime/AudioSynesthesia/Source/AudioSynesthesiaEditor/Private/AudioSynesthesiaEditorModule.cpp b/Engine/Plugins/Runtime/AudioSynesthesia/Source/AudioSynesthesiaEditor/Private/AudioSynesthesiaEditorModule.cpp
--- a/Engine/Plugins/Runtime/AudioSynesthesia/Source/AudioSynesthesiaEditor/Private/AudioSynesthesiaEditorModule.cpp
+++ b/Engine/Plugins/Runtime/AudioSynesthesia/Source/AudioSynesthesiaEditor/Private/AudioSynesthesiaEditorModule.cpp
@@ -58,7 +58,7 @@ private:
 
 			// look for Synesthesia classes 
 			UClass* ParentClass = ChildClass->GetSuperClass();
-			if (ParentClass->IsChildOf(SynesthesiaAssetType::StaticClass()))
+			if (ParentClass && ParentClass->IsChildOf(SynesthesiaAssetType::StaticClass()))
 			{
 				SynesthesiaAssetType* Synesthesia = ChildClass->GetDefaultObject<SynesthesiaAssetType>();
 				check(Synesthesia);
diff --git a/Engine/Source/Editor/AudioEditor/Private/AudioEditorModule.cpp b/Engine/Source/Editor/AudioEditor/Private/AudioEditorModule.cpp
--- a/Engine/Source/Editor/AudioEditor/Private/AudioEditorModule.cpp
+++ b/Engine/Source/Editor/AudioEditor/Private/AudioEditorModule.cpp
@@ -234,7 +234,7 @@ public:
 
 				// Look for submix or source preset classes
 				UClass* ParentClass = ChildClass->GetSuperClass();
-				if (ParentClass->IsChildOf(USoundEffectSourcePreset::StaticClass()) || ParentClass->IsChildOf(USoundEffectSubmixPreset::StaticClass()))
+				if (ParentClass && (ParentClass->IsChildOf(USoundEffectSourcePreset::StaticClass()) || ParentClass->IsChildOf(USoundEffectSubmixPreset::StaticClass())))
 				{
 					USoundEffectPreset* EffectPreset = ChildClass->GetDefaultObject<USoundEffectPreset>();
 					if (!RegisteredActions.Contains(EffectPreset) && EffectPreset->HasAssetActions())
