Fix 16-bit WAV import warning

diff --git a/Engine/Source/Editor/AudioEditor/Private/Factories/ReimportSoundFactory.cpp b/Engine/Source/Editor/AudioEditor/Private/Factories/ReimportSoundFactory.cpp
--- a/Engine/Source/Editor/AudioEditor/Private/Factories/ReimportSoundFactory.cpp
+++ b/Engine/Source/Editor/AudioEditor/Private/Factories/ReimportSoundFactory.cpp
@@ -29,6 +29,7 @@ UReimportSoundFactory::UReimportSoundFactory(const FObjectInitializer& ObjectIni
 	bIncludeModulatorNode = false;
 	bIncludeLoopingNode = false;
 	CueVolume = 0.75f;
+	ImportPriority = 1;
 }
 
 bool UReimportSoundFactory::CanReimport(UObject* Obj, TArray<FString>& OutFilenames)
