"Using Native Audio" Telemetry

diff --git a/Engine/Plugins/Runtime/Metasound/Source/MetasoundEditor/Private/MetasoundEditor.cpp b/Engine/Plugins/Runtime/Metasound/Source/MetasoundEditor/Private/MetasoundEditor.cpp
--- a/Engine/Plugins/Runtime/Metasound/Source/MetasoundEditor/Private/MetasoundEditor.cpp
+++ b/Engine/Plugins/Runtime/Metasound/Source/MetasoundEditor/Private/MetasoundEditor.cpp
@@ -11,6 +11,7 @@
 #include "EdGraph/EdGraphPin.h"
 #include "EdGraphUtilities.h"
 #include "Editor.h"
+#include "EngineAnalytics.h"
 #include "Styling/AppStyle.h"
 #include "Framework/Application/SlateApplication.h"
 #include "Framework/Commands/GenericCommands.h"
@@ -883,6 +884,11 @@ namespace Metasound
 			{
 				Play();
 			}
+
+			if (FEngineAnalytics::IsAvailable())
+			{
+				FEngineAnalytics::GetProvider().RecordEvent(TEXT("Audio.Usage.MetaSounds.EditorOpened"));
+			}
 		}
 
 		UObject* FEditor::GetMetasoundObject() const
diff --git a/Engine/Source/Editor/AudioEditor/Private/Factories/SoundFactory.cpp b/Engine/Source/Editor/AudioEditor/Private/Factories/SoundFactory.cpp
--- a/Engine/Source/Editor/AudioEditor/Private/Factories/SoundFactory.cpp
+++ b/Engine/Source/Editor/AudioEditor/Private/Factories/SoundFactory.cpp
@@ -5,6 +5,7 @@
 #include "Audio.h"
 #include "Components/AudioComponent.h"
 #include "ContentBrowserModule.h"
+#include "EngineAnalytics.h"
 #include "IContentBrowserSingleton.h"
 #include "Modules/ModuleManager.h"
 #include "Sound/SoundCue.h"
@@ -598,6 +599,11 @@ UObject* USoundFactory::CreateObject
 
 		Sound->SetRedrawThumbnail(true);
 
+		if (FEngineAnalytics::IsAvailable())
+		{
+			FEngineAnalytics::GetProvider().RecordEvent(TEXT("Audio.Usage.SoundFactory.SoundWaveImported"));
+		}
+
 		return Sound;
 	}
 
diff --git a/Engine/Source/Editor/AudioEditor/Private/SoundClassEditor.cpp b/Engine/Source/Editor/AudioEditor/Private/SoundClassEditor.cpp
--- a/Engine/Source/Editor/AudioEditor/Private/SoundClassEditor.cpp
+++ b/Engine/Source/Editor/AudioEditor/Private/SoundClassEditor.cpp
@@ -9,6 +9,7 @@
 #include "SoundClassGraph/SoundClassGraphSchema.h"
 #include "Sound/SoundClass.h"
 #include "Editor.h"
+#include "EngineAnalytics.h"
 
 #include "PropertyEditorModule.h"
 #include "IDetailsView.h"
@@ -174,6 +175,12 @@ void FSoundClassEditor::InitSoundClassEditor( const EToolkitMode::Type Mode, con
 			break;
 		}
 	}
+
+	if (FEngineAnalytics::IsAvailable())
+	{
+		FEngineAnalytics::GetProvider().RecordEvent(TEXT("Audio.Usage.SoundClass.EditorOpened"));
+	}
+
 }
 
 void FSoundClassEditor::BindCommands()
diff --git a/Engine/Source/Editor/AudioEditor/Private/SoundCueEditor.cpp b/Engine/Source/Editor/AudioEditor/Private/SoundCueEditor.cpp
--- a/Engine/Source/Editor/AudioEditor/Private/SoundCueEditor.cpp
+++ b/Engine/Source/Editor/AudioEditor/Private/SoundCueEditor.cpp
@@ -3,6 +3,7 @@
 #include "SoundCueEditor.h"
 #include "Framework/MultiBox/MultiBoxBuilder.h"
 #include "EdGraph/EdGraphNode.h"
+#include "EngineAnalytics.h"
 #include "Modules/ModuleManager.h"
 #include "Styling/AppStyle.h"
 #include "SoundCueGraph/SoundCueGraph.h"
@@ -163,6 +164,11 @@ void FSoundCueEditor::InitSoundCueEditor(const EToolkitMode::Type Mode, const TS
 		SpawnToolkitTab(GraphCanvasTabId, FString(), EToolkitTabSpot::Viewport);
 		SpawnToolkitTab(PropertiesTabId, FString(), EToolkitTabSpot::Details);
 	}*/
+
+	if (FEngineAnalytics::IsAvailable())
+	{
+		FEngineAnalytics::GetProvider().RecordEvent(TEXT("Audio.Usage.SoundCue.EditorOpened"));
+	}
 }
 
 USoundCue* FSoundCueEditor::GetSoundCue() const
diff --git a/Engine/Source/Editor/AudioEditor/Private/SoundSubmixEditor.cpp b/Engine/Source/Editor/AudioEditor/Private/SoundSubmixEditor.cpp
--- a/Engine/Source/Editor/AudioEditor/Private/SoundSubmixEditor.cpp
+++ b/Engine/Source/Editor/AudioEditor/Private/SoundSubmixEditor.cpp
@@ -7,6 +7,7 @@
 #include "AudioEditorModule.h"
 #include "EdGraph/EdGraph.h"
 #include "Editor.h"
+#include "EngineAnalytics.h"
 #include "GraphEditor.h"
 #include "GraphEditAction.h"
 #include "Kismet2/BlueprintEditorUtils.h"
@@ -27,7 +28,6 @@
 #include "UObject/Package.h"
 #include "Widgets/Docking/SDockTab.h"
 
-
 #define LOCTEXT_NAMESPACE "SoundSubmixEditor"
 DEFINE_LOG_CATEGORY_STATIC(LogSoundSubmixEditor, Log, All);
 
@@ -193,6 +193,12 @@ void FSoundSubmixEditor::Init(const EToolkitMode::Type Mode, const TSharedPtr<IT
 			break;
 		}
 	}
+
+	if (FEngineAnalytics::IsAvailable())
+	{
+		FEngineAnalytics::GetProvider().RecordEvent(TEXT("Audio.Usage.Submix.EditorOpened"));
+	}
+
 }
 
 FSoundSubmixEditor::~FSoundSubmixEditor()
