Fix Non-16-bit WAV Asset Import

diff --git a/Engine/Source/Editor/AudioEditor/Private/Factories/SoundFactory.cpp b/Engine/Source/Editor/AudioEditor/Private/Factories/SoundFactory.cpp
--- a/Engine/Source/Editor/AudioEditor/Private/Factories/SoundFactory.cpp
+++ b/Engine/Source/Editor/AudioEditor/Private/Factories/SoundFactory.cpp
@@ -531,7 +531,7 @@ UObject* USoundFactory::CreateObject
 		else
 		{
 			// For mono and stereo assets, just copy the data into the buffer
-			FSharedBuffer UpdatedBuffer = FSharedBuffer::Clone(RawWaveData.GetData(), RawWaveData.Num());
+			FSharedBuffer UpdatedBuffer = FSharedBuffer::Clone(Buffer, RawWaveDataBufferSize);
 			Sound->RawData.UpdatePayload(UpdatedBuffer);
 
 		}
