Increased audio LLM tag coverage

diff --git a/Engine/Plugins/Runtime/AudioCapture/Source/AudioCapture/Private/AudioCapture.cpp b/Engine/Plugins/Runtime/AudioCapture/Source/AudioCapture/Private/AudioCapture.cpp
--- a/Engine/Plugins/Runtime/AudioCapture/Source/AudioCapture/Private/AudioCapture.cpp
+++ b/Engine/Plugins/Runtime/AudioCapture/Source/AudioCapture/Private/AudioCapture.cpp
@@ -1,11 +1,13 @@
 // Copyright Epic Games, Inc. All Rights Reserved.
 
 #include "AudioCapture.h"
+
+#include "AudioCaptureCore.h"
 #include "CoreMinimal.h"
+#include "HAL/LowLevelMemTracker.h"
+#include "Misc/ConfigCacheIni.h"
 #include "Misc/ScopeLock.h"
 #include "Modules/ModuleManager.h"
-#include "AudioCaptureCore.h"
-#include "Misc/ConfigCacheIni.h"
 
 DEFINE_LOG_CATEGORY(LogAudioCapture);
 
@@ -81,6 +83,7 @@ bool UAudioCapture::IsCapturingAudio()
 
 UAudioCapture* UAudioCaptureFunctionLibrary::CreateAudioCapture()
 {
+	LLM_SCOPE(ELLMTag::Audio);
 	UAudioCapture* NewAudioCapture = NewObject<UAudioCapture>();
 	if (NewAudioCapture->OpenDefaultAudioStream())
 	{
@@ -93,6 +96,7 @@ UAudioCapture* UAudioCaptureFunctionLibrary::CreateAudioCapture()
 
 void FAudioCaptureModule::StartupModule()
 {
+	LLM_SCOPE(ELLMTag::Audio);
 	// Load platform specific implementations for audio capture (if specified in a .ini file)
 	FString AudioCaptureModuleName;
 	if (GConfig->GetString(TEXT("Audio"), TEXT("AudioCaptureModuleName"), AudioCaptureModuleName, GEngineIni))
@@ -106,4 +110,4 @@ void FAudioCaptureModule::ShutdownModule()
 }
 
 
-IMPLEMENT_MODULE(FAudioCaptureModule, AudioCapture);
\ No newline at end of file
+IMPLEMENT_MODULE(FAudioCaptureModule, AudioCapture);
diff --git a/Engine/Plugins/Runtime/AudioModulation/Source/AudioModulation/Private/AudioModulation.cpp b/Engine/Plugins/Runtime/AudioModulation/Source/AudioModulation/Private/AudioModulation.cpp
--- a/Engine/Plugins/Runtime/AudioModulation/Source/AudioModulation/Private/AudioModulation.cpp
+++ b/Engine/Plugins/Runtime/AudioModulation/Source/AudioModulation/Private/AudioModulation.cpp
@@ -7,6 +7,7 @@
 #include "AudioModulationSystem.h"
 #include "CanvasTypes.h"
 #include "Features/IModularFeatures.h"
+#include "HAL/LowLevelMemTracker.h"
 #include "IAudioModulation.h"
 #include "Modules/ModuleManager.h"
 #include "SoundControlBusMix.h"
@@ -256,6 +257,7 @@ TAudioModulationPtr FAudioModulationPluginFactory::CreateNewModulationPlugin(FAu
 
 void FAudioModulationModule::StartupModule()
 {
+	LLM_SCOPE(ELLMTag::AudioMixerPlugins);
 	IModularFeatures::Get().RegisterModularFeature(FAudioModulationPluginFactory::GetModularFeatureName(), &ModulationPluginFactory);
 
 	if (const UAudioModulationSettings* ModulationSettings = GetDefault<UAudioModulationSettings>())
@@ -274,6 +276,7 @@ void FAudioModulationModule::StartupModule()
 
 void FAudioModulationModule::ShutdownModule()
 {
+	LLM_SCOPE(ELLMTag::AudioMixerPlugins);
 	IModularFeatures::Get().UnregisterModularFeature(FAudioModulationPluginFactory::GetModularFeatureName(), &ModulationPluginFactory);
 	UE_LOG(LogAudioModulation, Log, TEXT("Audio Modulation Shutdown"));
 }
diff --git a/Engine/Plugins/Runtime/AudioSynesthesia/Source/AudioSynesthesia/Private/AudioSynesthesiaModule.cpp b/Engine/Plugins/Runtime/AudioSynesthesia/Source/AudioSynesthesia/Private/AudioSynesthesiaModule.cpp
--- a/Engine/Plugins/Runtime/AudioSynesthesia/Source/AudioSynesthesia/Private/AudioSynesthesiaModule.cpp
+++ b/Engine/Plugins/Runtime/AudioSynesthesia/Source/AudioSynesthesia/Private/AudioSynesthesiaModule.cpp
@@ -1,12 +1,14 @@
 // Copyright Epic Games, Inc. All Rights Reserved.
 
 #include "AudioSynesthesiaModule.h"
+
+#include "AudioAnalyzerModule.h"
+#include "AudioSynesthesiaLog.h"
 #include "CoreMinimal.h"
-#include "Modules/ModuleInterface.h"
 #include "Features/IModularFeatures.h"
+#include "HAL/LowLevelMemTracker.h"
 #include "LoudnessNRTFactory.h"
-#include "AudioAnalyzerModule.h"
-#include "AudioSynesthesiaLog.h"
+#include "Modules/ModuleInterface.h"
 
 DEFINE_LOG_CATEGORY(LogAudioSynesthesia);
 
@@ -16,6 +18,7 @@ namespace Audio
 {
 	void FAudioSynesthesiaModule::StartupModule()
 	{
+		LLM_SCOPE_BYTAG(AudioAnalysis);
 		FModuleManager::Get().LoadModuleChecked<IModuleInterface>("AudioSynesthesiaCore");
 		FModuleManager::Get().LoadModuleChecked<FAudioAnalyzerModule>("AudioAnalyzer");
 	}
diff --git a/Engine/Plugins/Runtime/AudioSynesthesia/Source/AudioSynesthesiaCore/Private/AudioSynesthesiaCoreModule.cpp b/Engine/Plugins/Runtime/AudioSynesthesia/Source/AudioSynesthesiaCore/Private/AudioSynesthesiaCoreModule.cpp
--- a/Engine/Plugins/Runtime/AudioSynesthesia/Source/AudioSynesthesiaCore/Private/AudioSynesthesiaCoreModule.cpp
+++ b/Engine/Plugins/Runtime/AudioSynesthesia/Source/AudioSynesthesiaCore/Private/AudioSynesthesiaCoreModule.cpp
@@ -1,6 +1,7 @@
 // Copyright Epic Games, Inc. All Rights Reserved.
 #include "AudioSynesthesiaCoreModule.h"
 
+#include "AudioAnalyzerModule.h"
 #include "AudioSynesthesiaCoreLog.h"
 #include "ConstantQNRTFactory.h"
 #include "CoreMinimal.h"
@@ -21,6 +22,7 @@ namespace Audio
 		public:
 			void StartupModule()
 			{
+				LLM_SCOPE_BYNAME(TEXT("Audio/AudioAnalysis"));
 				// Register factories on startup
 				IModularFeatures::Get().RegisterModularFeature(FLoudnessNRTFactory::GetModularFeatureName(), &LoudnessNRTFactory);
 				IModularFeatures::Get().RegisterModularFeature(FConstantQNRTFactory::GetModularFeatureName(), &ConstantQNRTFactory);
@@ -32,6 +34,7 @@ namespace Audio
 
 			void ShutdownModule()
 			{
+				LLM_SCOPE_BYNAME(TEXT("Audio/AudioAnalysis"));
 				// Unregister factories on shutdown
 				IModularFeatures::Get().UnregisterModularFeature(FLoudnessNRTFactory::GetModularFeatureName(), &LoudnessNRTFactory);
 				IModularFeatures::Get().UnregisterModularFeature(FConstantQNRTFactory::GetModularFeatureName(), &ConstantQNRTFactory);
@@ -49,7 +52,6 @@ namespace Audio
 			FLoudnessFactory LoudnessFactory;
 			FMeterFactory MeterFactory;
 	};
-
 }
 
 IMPLEMENT_MODULE(Audio::FAudioSynesthesiaCoreModule, AudioSynesthesiaCore);
diff --git a/Engine/Plugins/Runtime/AudioSynesthesia/Source/AudioSynesthesiaEditor/Private/AudioSynesthesiaEditorModule.cpp b/Engine/Plugins/Runtime/AudioSynesthesia/Source/AudioSynesthesiaEditor/Private/AudioSynesthesiaEditorModule.cpp
--- a/Engine/Plugins/Runtime/AudioSynesthesia/Source/AudioSynesthesiaEditor/Private/AudioSynesthesiaEditorModule.cpp
+++ b/Engine/Plugins/Runtime/AudioSynesthesia/Source/AudioSynesthesiaEditor/Private/AudioSynesthesiaEditorModule.cpp
@@ -1,15 +1,18 @@
 // Copyright Epic Games, Inc. All Rights Reserved.
 
 #include "AudioSynesthesiaEditorModule.h"
-#include "CoreMinimal.h"
-#include "Modules/ModuleManager.h"
+
 #include "AssetToolsModule.h"
-#include "UObject/UObjectIterator.h"
-#include "AudioSynesthesiaNRT.h"
-#include "AudioSynesthesia.h"
 #include "AssetTypeActions_AudioSynesthesiaNRT.h"
 #include "AssetTypeActions_AudioSynesthesiaNRTSettings.h"
 #include "AssetTypeActions_AudioSynesthesiaSettings.h"
+#include "AudioAnalyzerModule.h"
+#include "AudioSynesthesia.h"
+#include "AudioSynesthesiaNRT.h"
+#include "CoreMinimal.h"
+#include "HAL/LowLevelMemTracker.h"
+#include "Modules/ModuleManager.h"
+#include "UObject/UObjectIterator.h"
 
 
 DEFINE_LOG_CATEGORY(LogAudioSynesthesiaEditor);
@@ -22,6 +25,7 @@ public:
 
 	virtual void StartupModule() override
 	{
+		LLM_SCOPE_BYTAG(AudioAnalysis);
 		RegisterAssetActions();
 	}
 
diff --git a/Engine/Plugins/Runtime/Metasound/Source/MetasoundEditor/Private/MetasoundEditorModule.cpp b/Engine/Plugins/Runtime/Metasound/Source/MetasoundEditor/Private/MetasoundEditorModule.cpp
--- a/Engine/Plugins/Runtime/Metasound/Source/MetasoundEditor/Private/MetasoundEditorModule.cpp
+++ b/Engine/Plugins/Runtime/Metasound/Source/MetasoundEditor/Private/MetasoundEditorModule.cpp
@@ -32,6 +32,7 @@
 #include "MetasoundSettings.h"
 #include "MetasoundSource.h"
 #include "MetasoundTime.h"
+#include "MetasoundTrace.h"
 #include "MetasoundTrigger.h"
 #include "MetasoundUObjectRegistry.h"
 #include "Modules/ModuleInterface.h"
@@ -651,6 +652,7 @@ namespace Metasound
 
 			virtual void StartupModule() override
 			{
+				METASOUND_LLM_SCOPE;
 				// Register Metasound asset type actions
 				IAssetTools& AssetTools = FModuleManager::LoadModuleChecked<FAssetToolsModule>(AssetToolName).Get();
 
@@ -756,6 +758,8 @@ namespace Metasound
 
 			virtual void ShutdownModule() override
 			{
+				METASOUND_LLM_SCOPE;
+
 				if (ISettingsModule* SettingsModule = FModuleManager::GetModulePtr<ISettingsModule>("Settings"))
 				{
 					SettingsModule->UnregisterSettings("Editor", "Audio", "MetaSound Editor");
diff --git a/Engine/Plugins/Runtime/Metasound/Source/MetasoundEngine/Private/MetasoundEngineModule.cpp b/Engine/Plugins/Runtime/Metasound/Source/MetasoundEngine/Private/MetasoundEngineModule.cpp
--- a/Engine/Plugins/Runtime/Metasound/Source/MetasoundEngine/Private/MetasoundEngineModule.cpp
+++ b/Engine/Plugins/Runtime/Metasound/Source/MetasoundEngine/Private/MetasoundEngineModule.cpp
@@ -8,6 +8,7 @@
 #include "MetasoundEngineArchetypes.h"
 #include "MetasoundInterface.h"
 #include "MetasoundSource.h"
+#include "MetasoundTrace.h"
 #include "MetasoundUObjectRegistry.h"
 #include "MetasoundWave.h"
 #include "MetasoundWaveTable.h"
@@ -26,6 +27,7 @@ class FMetasoundEngineModule : public IMetasoundEngineModule
 {
 	virtual void StartupModule() override
 	{
+		METASOUND_LLM_SCOPE;
 		FModuleManager::Get().LoadModuleChecked("MetasoundGraphCore");
 		FModuleManager::Get().LoadModuleChecked("MetasoundFrontend");
 		FModuleManager::Get().LoadModuleChecked("MetasoundStandardNodes");
diff --git a/Engine/Plugins/Runtime/Metasound/Source/MetasoundFrontend/Private/MetasoundFrontendDocument.cpp b/Engine/Plugins/Runtime/Metasound/Source/MetasoundFrontend/Private/MetasoundFrontendDocument.cpp
--- a/Engine/Plugins/Runtime/Metasound/Source/MetasoundFrontend/Private/MetasoundFrontendDocument.cpp
+++ b/Engine/Plugins/Runtime/Metasound/Source/MetasoundFrontend/Private/MetasoundFrontendDocument.cpp
@@ -216,6 +216,10 @@ FMetasoundFrontendClassInterface FMetasoundFrontendClassInterface::GenerateClass
 		FMetasoundFrontendInterfaceStyle InputStyle;
 #endif // WITH_EDITOR
 
+		// Reserve memory to minimize memory use in ClassInterface.Inputs array.
+		ClassInterface.Inputs.Reserve(InputInterface.Num());
+
+
 		for (const FInputDataVertex& InputVertex : InputInterface)
 		{
 			FMetasoundFrontendClassInput ClassInput;
@@ -270,6 +274,9 @@ FMetasoundFrontendClassInterface FMetasoundFrontendClassInterface::GenerateClass
 		FMetasoundFrontendInterfaceStyle OutputStyle;
 #endif // WITH_EDITOR
 
+		// Reserve memory to minimize memory use in ClassInterface.Outputs array.
+		ClassInterface.Outputs.Reserve(OutputInterface.Num());
+
 		for (const FOutputDataVertex& OutputVertex: OutputInterface)
 		{
 			FMetasoundFrontendClassOutput ClassOutput;
@@ -309,6 +316,9 @@ FMetasoundFrontendClassInterface FMetasoundFrontendClassInterface::GenerateClass
 #endif // WITH_EDITOR
 	}
 
+	// Reserve size to minimize memory use in ClassInterface.Environment array
+	ClassInterface.Environment.Reserve(InVertexInterface.GetEnvironmentInterface().Num());
+
 	for (const FEnvironmentVertex& EnvVertex : InVertexInterface.GetEnvironmentInterface())
 	{
 		FMetasoundFrontendClassEnvironmentVariable EnvVar;
diff --git a/Engine/Plugins/Runtime/Metasound/Source/MetasoundGraphCore/Private/MetasoundTrace.cpp b/Engine/Plugins/Runtime/Metasound/Source/MetasoundGraphCore/Private/MetasoundTrace.cpp
new file mode 100644
--- /dev/null
+++ b/Engine/Plugins/Runtime/Metasound/Source/MetasoundGraphCore/Private/MetasoundTrace.cpp
@@ -0,0 +1,10 @@
+// Copyright Epic Games, Inc. All Rights Reserved.
+
+#include "MetasoundTrace.h"
+
+#include "HAL/LowLevelMemStats.h"
+#include "HAL/LowLevelMemTracker.h"
+
+DECLARE_LLM_MEMORY_STAT(TEXT("MetaSound"), STAT_MetaSoundLLM, STATGROUP_LLMFULL);
+
+LLM_DEFINE_TAG(MetaSound, TEXT("MetaSound"), TEXT("Audio"), GET_STATFNAME(STAT_MetaSoundLLM), GET_STATFNAME(STAT_AudioSummaryLLM));
diff --git a/Engine/Plugins/Runtime/Metasound/Source/MetasoundGraphCore/Public/MetasoundTrace.h b/Engine/Plugins/Runtime/Metasound/Source/MetasoundGraphCore/Public/MetasoundTrace.h
--- a/Engine/Plugins/Runtime/Metasound/Source/MetasoundGraphCore/Public/MetasoundTrace.h
+++ b/Engine/Plugins/Runtime/Metasound/Source/MetasoundGraphCore/Public/MetasoundTrace.h
@@ -2,6 +2,7 @@
 
 #pragma once
 
+#include "HAL/LowLevelMemTracker.h"
 #include "ProfilingDebugging/CpuProfilerTrace.h"
 
 #if !defined(METASOUND_CPUPROFILERTRACE_ENABLED)
@@ -29,6 +30,8 @@
 
 #endif
 
+LLM_DECLARE_TAG_API(MetaSound, METASOUNDGRAPHCORE_API);
 // Convenience macro for MetaSound LLM scope to avoid misspells. 
-#define METASOUND_LLM_SCOPE LLM_SCOPE_BYNAME(TEXT("Audio/MetaSound"));
+//
+#define METASOUND_LLM_SCOPE LLM_SCOPE_BYTAG(MetaSound);
 
diff --git a/Engine/Plugins/Runtime/ResonanceAudio/Source/ResonanceAudio/Private/ResonanceAudioModule.cpp b/Engine/Plugins/Runtime/ResonanceAudio/Source/ResonanceAudio/Private/ResonanceAudioModule.cpp
--- a/Engine/Plugins/Runtime/ResonanceAudio/Source/ResonanceAudio/Private/ResonanceAudioModule.cpp
+++ b/Engine/Plugins/Runtime/ResonanceAudio/Source/ResonanceAudio/Private/ResonanceAudioModule.cpp
@@ -4,11 +4,13 @@
 
 #include "ResonanceAudioModule.h"
 
+#include "AudioDevice.h"
+#include "Features/IModularFeatures.h"
+#include "HAL/LowLevelMemTracker.h"
 #include "ResonanceAudioPluginListener.h"
 #include "ResonanceAudioReverb.h"
-#include "ResonanceAudioSpatialization.h"
-#include "Features/IModularFeatures.h"
 #include "ResonanceAudioSettings.h"
+#include "ResonanceAudioSpatialization.h"
 
 IMPLEMENT_MODULE(ResonanceAudio::FResonanceAudioModule, ResonanceAudio)
 
@@ -29,6 +31,7 @@ namespace ResonanceAudio
 
 	void FResonanceAudioModule::StartupModule()
 	{
+		LLM_SCOPE_BYTAG(AudioSpatializationPlugins);
 		check(bModuleInitialized == false);
 		bModuleInitialized = true;
 
@@ -59,6 +62,7 @@ namespace ResonanceAudio
 
 	void FResonanceAudioModule::ShutdownModule()
 	{
+		LLM_SCOPE_BYTAG(AudioSpatializationPlugins);
 		check(bModuleInitialized == true);
 		bModuleInitialized = false;
 		UE_LOG(LogResonanceAudio, Log, TEXT("Resonance Audio Module is Shutdown"));
@@ -81,6 +85,7 @@ namespace ResonanceAudio
 
 	void FResonanceAudioModule::RegisterAudioDevice(FAudioDevice* AudioDeviceHandle)
 	{
+		LLM_SCOPE_BYTAG(AudioSpatializationPlugins);
 		if (!RegisteredAudioDevices.Contains(AudioDeviceHandle))
 		{
 			TAudioPluginListenerPtr NewResonanceAudioPluginListener = TAudioPluginListenerPtr(new FResonanceAudioPluginListener());
@@ -91,6 +96,7 @@ namespace ResonanceAudio
 
 	void FResonanceAudioModule::UnregisterAudioDevice(FAudioDevice* AudioDeviceHandle)
 	{
+		LLM_SCOPE_BYTAG(AudioSpatializationPlugins);
 		RegisteredAudioDevices.Remove(AudioDeviceHandle);
 		UE_LOG(LogResonanceAudio, Log, TEXT("Audio Device unregistered from Resonance"));
 	}
@@ -102,6 +108,7 @@ namespace ResonanceAudio
 
 	TAudioSpatializationPtr FSpatializationPluginFactory::CreateNewSpatializationPlugin(FAudioDevice* OwningDevice)
 	{
+		LLM_SCOPE_BYTAG(AudioSpatializationPlugins);
 		// Register the audio device to the Resonance Audio module.
 		FResonanceAudioModule* Module = &FModuleManager::GetModuleChecked<FResonanceAudioModule>("ResonanceAudio");
 		if (Module != nullptr)
@@ -113,6 +120,7 @@ namespace ResonanceAudio
 
 	TAudioReverbPtr FReverbPluginFactory::CreateNewReverbPlugin(FAudioDevice* OwningDevice)
 	{
+		LLM_SCOPE_BYTAG(AudioSpatializationPlugins);
 		// Register the audio device to the Resonance Audio module.
 		FResonanceAudioModule* Module = &FModuleManager::GetModuleChecked<FResonanceAudioModule>("ResonanceAudio");
 		if (Module != nullptr)
diff --git a/Engine/Plugins/Runtime/ResonanceAudio/Source/ResonanceAudio/Public/ResonanceAudioAmbisonicsSettings.h b/Engine/Plugins/Runtime/ResonanceAudio/Source/ResonanceAudio/Public/ResonanceAudioAmbisonicsSettings.h
--- a/Engine/Plugins/Runtime/ResonanceAudio/Source/ResonanceAudio/Public/ResonanceAudioAmbisonicsSettings.h
+++ b/Engine/Plugins/Runtime/ResonanceAudio/Source/ResonanceAudio/Public/ResonanceAudioAmbisonicsSettings.h
@@ -4,10 +4,13 @@
 
 #pragma once
 
-#include "ResonanceAudioCommon.h"
+#include "AudioDevice.h"
 #include "ISoundfieldFormat.h"
+#include "ResonanceAudioCommon.h"
+#include "HAL/LowLevelMemTracker.h"
 #include "ResonanceAudioAmbisonicsSettings.generated.h"
 
+
 UENUM()
 enum class EResonanceRenderMode : uint8
 {
@@ -45,6 +48,7 @@ public:
 
 	virtual TUniquePtr<ISoundfieldEncodingSettingsProxy> Duplicate() const override
 	{
+		LLM_SCOPE_BYTAG(AudioSpatializationPlugins);
 		FResonanceAmbisonicsSettingsProxy* Proxy = new FResonanceAmbisonicsSettingsProxy();
 		Proxy->RenderingMode = RenderingMode;
 		return TUniquePtr<ISoundfieldEncodingSettingsProxy>(Proxy);
@@ -63,6 +67,7 @@ public:
 
 	virtual TUniquePtr<ISoundfieldEncodingSettingsProxy> GetProxy() const override
 	{
+		LLM_SCOPE_BYTAG(AudioSpatializationPlugins);
 		FResonanceAmbisonicsSettingsProxy* Proxy = new FResonanceAmbisonicsSettingsProxy();
 
 		switch (RenderMode)
@@ -101,4 +106,4 @@ public:
 		return TUniquePtr<ISoundfieldEncodingSettingsProxy>(Proxy);
 	}
 
-};
\ No newline at end of file
+};
diff --git a/Engine/Plugins/Runtime/SoundCueTemplates/Source/SoundCueTemplatesEditor/Private/SoundCueTemplatesEditorModule.cpp b/Engine/Plugins/Runtime/SoundCueTemplates/Source/SoundCueTemplatesEditor/Private/SoundCueTemplatesEditorModule.cpp
--- a/Engine/Plugins/Runtime/SoundCueTemplates/Source/SoundCueTemplatesEditor/Private/SoundCueTemplatesEditorModule.cpp
+++ b/Engine/Plugins/Runtime/SoundCueTemplates/Source/SoundCueTemplatesEditor/Private/SoundCueTemplatesEditorModule.cpp
@@ -5,23 +5,27 @@
 #include "Modules/ModuleManager.h"
 #include "AssetToolsModule.h"
 #include "ToolMenus.h"
+#include "HAL/LowLevelMemTracker.h"
 
 #include "AssetTypeActions_SoundCueTemplate.h"
 
 void FSoundCueTemplatesEditorModule::StartupModule()
 {
+	LLM_SCOPE(ELLMTag::AudioMisc);
 	RegisterAssetActions();
 	UToolMenus::RegisterStartupCallback(FSimpleMulticastDelegate::FDelegate::CreateRaw(this, &FSoundCueTemplatesEditorModule::RegisterMenus));
 }
 
 void FSoundCueTemplatesEditorModule::ShutdownModule()
 {
+	LLM_SCOPE(ELLMTag::AudioMisc);
 	UToolMenus::UnRegisterStartupCallback(this);
 	UToolMenus::UnregisterOwner("SoundCueTemplates");
 }
 
 void FSoundCueTemplatesEditorModule::RegisterAssetActions()
 {
+	LLM_SCOPE(ELLMTag::AudioMisc);
 	// Register the audio editor asset type actions
 	IAssetTools& AssetTools = FModuleManager::LoadModuleChecked<FAssetToolsModule>("AssetTools").Get();
 
diff --git a/Engine/Plugins/Runtime/SoundUtilities/Source/SoundUtilitiesEditor/Private/SoundUtilitiesEditorModule.cpp b/Engine/Plugins/Runtime/SoundUtilities/Source/SoundUtilitiesEditor/Private/SoundUtilitiesEditorModule.cpp
--- a/Engine/Plugins/Runtime/SoundUtilities/Source/SoundUtilitiesEditor/Private/SoundUtilitiesEditorModule.cpp
+++ b/Engine/Plugins/Runtime/SoundUtilities/Source/SoundUtilitiesEditor/Private/SoundUtilitiesEditorModule.cpp
@@ -11,11 +11,13 @@
 #include "ToolMenus.h"
 #include "AudioEditorModule.h"
 #include "SoundWaveAssetActionExtender.h"
+#include "HAL/LowLevelMemTracker.h"
 
 IMPLEMENT_MODULE(FSoundUtilitiesEditorModule, SoundUtilitiesEditor)
 
 void FSoundUtilitiesEditorModule::StartupModule()
 {
+	LLM_SCOPE(ELLMTag::AudioMisc);
 	IAssetTools& AssetTools = FModuleManager::LoadModuleChecked<FAssetToolsModule>("AssetTools").Get();
 	
 	// Register asset actions
@@ -26,11 +28,13 @@ void FSoundUtilitiesEditorModule::StartupModule()
 
 void FSoundUtilitiesEditorModule::ShutdownModule()
 {
+	LLM_SCOPE(ELLMTag::AudioMisc);
 	UToolMenus::UnRegisterStartupCallback(this);
 	UToolMenus::UnregisterOwner("SoundUtilities");
 }
 
 void FSoundUtilitiesEditorModule::RegisterMenus()
 {
+	LLM_SCOPE(ELLMTag::AudioMisc);
 	FSoundWaveAssetActionExtender::RegisterMenus();
 }
diff --git a/Engine/Plugins/Runtime/Synthesis/Source/Synthesis/Private/SynthesisModule.cpp b/Engine/Plugins/Runtime/Synthesis/Source/Synthesis/Private/SynthesisModule.cpp
--- a/Engine/Plugins/Runtime/Synthesis/Source/Synthesis/Private/SynthesisModule.cpp
+++ b/Engine/Plugins/Runtime/Synthesis/Source/Synthesis/Private/SynthesisModule.cpp
@@ -1,11 +1,13 @@
 // Copyright Epic Games, Inc. All Rights Reserved.
 
 #include "SynthesisModule.h"
+
 #include "CoreMinimal.h"
-#include "Stats/Stats.h"
-#include "UI/SynthSlateStyle.h"
+#include "HAL/LowLevelMemTracker.h"
 #include "Modules/ModuleInterface.h"
 #include "Modules/ModuleManager.h"
+#include "Stats/Stats.h"
+#include "UI/SynthSlateStyle.h"
 
 DEFINE_LOG_CATEGORY(LogSynthesis);
 
@@ -13,6 +15,7 @@ IMPLEMENT_MODULE(FSynthesisModule, Synthesis)
 
 void FSynthesisModule::StartupModule()
 {
+	LLM_SCOPE(ELLMTag::AudioSynthesis);
 	FSynthSlateStyleSet::Initialize();
 
 	FModuleManager::Get().LoadModuleChecked(TEXT("SignalProcessing"));
@@ -20,6 +23,7 @@ void FSynthesisModule::StartupModule()
 
 void FSynthesisModule::ShutdownModule()
 {
+	LLM_SCOPE(ELLMTag::AudioSynthesis);
 	FSynthSlateStyleSet::Shutdown();
 }
 
diff --git a/Engine/Plugins/Runtime/Synthesis/Source/SynthesisEditor/Private/SynthesisEditorModule.cpp b/Engine/Plugins/Runtime/Synthesis/Source/SynthesisEditor/Private/SynthesisEditorModule.cpp
--- a/Engine/Plugins/Runtime/Synthesis/Source/SynthesisEditor/Private/SynthesisEditorModule.cpp
+++ b/Engine/Plugins/Runtime/Synthesis/Source/SynthesisEditor/Private/SynthesisEditorModule.cpp
@@ -36,6 +36,7 @@
 #include "SubmixEffects/SubmixEffectStereoDelay.h"
 #include "SubmixEffects/SubmixEffectTapDelay.h"
 #include "SynthesisEditorSettings.h"
+#include "HAL/LowLevelMemTracker.h"
 
 
 DEFINE_LOG_CATEGORY(LogSynthesisEditor);
@@ -45,6 +46,7 @@ IMPLEMENT_MODULE(FSynthesisEditorModule, SynthesisEditor)
 
 void FSynthesisEditorModule::StartupModule()
 {
+	LLM_SCOPE(ELLMTag::AudioSynthesis);
 	IAssetTools& AssetTools = FModuleManager::LoadModuleChecked<FAssetToolsModule>("AssetTools").Get();
 
 	AssetTools.RegisterAssetTypeActions(MakeShared<FAssetTypeActions_ModularSynthPresetBank>());
@@ -64,5 +66,6 @@ void FSynthesisEditorModule::ShutdownModule()
 
 void FSynthesisEditorModule::RegisterMenus()
 {
+	LLM_SCOPE(ELLMTag::AudioSynthesis);
 	FAudioImpulseResponseExtension::RegisterMenus();
 }
diff --git a/Engine/Source/Editor/AudioEditor/Private/AudioEditorModule.cpp b/Engine/Source/Editor/AudioEditor/Private/AudioEditorModule.cpp
--- a/Engine/Source/Editor/AudioEditor/Private/AudioEditorModule.cpp
+++ b/Engine/Source/Editor/AudioEditor/Private/AudioEditorModule.cpp
@@ -22,6 +22,7 @@
 #include "EdGraphUtilities.h"
 #include "Factories/ReimportSoundFactory.h"
 #include "Factories/SoundFactory.h"
+#include "HAL/LowLevelMemTracker.h"
 #include "Modules/ModuleManager.h"
 #include "PropertyEditorModule.h"
 #include "Sound/AudioSettings.h"
@@ -79,12 +80,14 @@ class FAudioEditorModule : public IAudioEditorModule
 public:
 	FAudioEditorModule()
 	{
+		LLM_SCOPE(ELLMTag::AudioMisc);
 		// Create style set for audio asset icons
 		AudioStyleSet = MakeShared<FSlateStyleSet>("AudioStyleSet");
 	}
 
 	virtual void StartupModule() override
 	{
+		LLM_SCOPE(ELLMTag::AudioMisc);
 		SoundClassExtensibility.Init();
 		SoundCueExtensibility.Init();
 		SoundSubmixExtensibility.Init();
@@ -139,6 +142,7 @@ public:
 
 	virtual void ShutdownModule() override
 	{
+		LLM_SCOPE(ELLMTag::AudioMisc);
 #if WITH_SNDFILE_IO
 		Audio::ShutdownSoundFileIOManager();
 #endif // WITH_SNDFILE_IO
@@ -177,6 +181,7 @@ public:
 
 	virtual void RegisterAssetActions() override
 	{
+		LLM_SCOPE(ELLMTag::AudioMisc);
 		// Register the audio editor asset type actions
 		IAssetTools& AssetTools = FModuleManager::LoadModuleChecked<FAssetToolsModule>("AssetTools").Get();
 
@@ -249,6 +254,7 @@ public:
 
 	virtual TSharedRef<FAssetEditorToolkit> CreateSoundClassEditor( const EToolkitMode::Type Mode, const TSharedPtr< IToolkitHost >& InitToolkitHost, USoundClass* InSoundClass ) override
 	{
+		LLM_SCOPE(ELLMTag::AudioMisc);
 		TSharedRef<FSoundClassEditor> NewSoundClassEditor(new FSoundClassEditor());
 		NewSoundClassEditor->InitSoundClassEditor(Mode, InitToolkitHost, InSoundClass);
 		return NewSoundClassEditor;
@@ -256,6 +262,7 @@ public:
 
 	virtual TSharedRef<FAssetEditorToolkit> CreateSoundSubmixEditor(const EToolkitMode::Type Mode, const TSharedPtr< IToolkitHost >& InitToolkitHost, USoundSubmixBase* InSoundSubmix) override
 	{
+		LLM_SCOPE(ELLMTag::AudioMisc);
 		TSharedPtr<FSoundSubmixEditor> NewSubmixEditor = MakeShared<FSoundSubmixEditor>();
 		NewSubmixEditor->Init(Mode, InitToolkitHost, InSoundSubmix);
 		return StaticCastSharedPtr<FAssetEditorToolkit>(NewSubmixEditor).ToSharedRef();
@@ -283,6 +290,7 @@ public:
 
 	virtual TSharedRef<ISoundCueEditor> CreateSoundCueEditor(const EToolkitMode::Type Mode, const TSharedPtr< IToolkitHost >& InitToolkitHost, USoundCue* SoundCue) override
 	{
+		LLM_SCOPE(ELLMTag::AudioMisc);
 		TSharedRef<FSoundCueEditor> NewSoundCueEditor(new FSoundCueEditor());
 		NewSoundCueEditor->InitSoundCueEditor(Mode, InitToolkitHost, SoundCue);
 		return NewSoundCueEditor;
@@ -335,6 +343,7 @@ public:
 
 	USoundWave* ImportSoundWave(UPackage* const SoundWavePackage, const FString& InSoundWaveAssetName, const FString& InWavFilename) override
 	{
+		LLM_SCOPE(ELLMTag::AudioMisc);
 		USoundFactory* SoundWaveFactory = NewObject<USoundFactory>();
 
 		// Setup sane defaults for importing localized sound waves
diff --git a/Engine/Source/Runtime/AudioAnalyzer/Private/AudioAnalyzerFacade.cpp b/Engine/Source/Runtime/AudioAnalyzer/Private/AudioAnalyzerFacade.cpp
--- a/Engine/Source/Runtime/AudioAnalyzer/Private/AudioAnalyzerFacade.cpp
+++ b/Engine/Source/Runtime/AudioAnalyzer/Private/AudioAnalyzerFacade.cpp
@@ -10,6 +10,8 @@ namespace Audio
 {
 	IAnalyzerFactory* GetAnalyzerFactory(FName InFactoryName)
 	{
+		LLM_SCOPE_BYTAG(AudioAnalysis);
+
 		// Get all analyzer nrt factories implementations.
 		IModularFeatures::Get().LockModularFeatureList();
 		TArray<IAnalyzerFactory*> RegisteredFactories = IModularFeatures::Get().GetModularFeatureImplementations<IAnalyzerFactory>(IAnalyzerFactory::GetModularFeatureName());
@@ -43,6 +45,8 @@ namespace Audio
 
 	TUniquePtr<IAnalyzerResult> FAnalyzerFacade::AnalyzeAudioBuffer(const TArray<float>& InAudioBuffer, int32 InNumChannels, float InSampleRate)
 	{
+		LLM_SCOPE_BYTAG(AudioAnalysis);
+
 		if (nullptr == Factory)
 		{
 			UE_LOG(LogAudioAnalyzer, Error, TEXT("Cannot analyze audio due to null factory"));
diff --git a/Engine/Source/Runtime/AudioAnalyzer/Private/AudioAnalyzerModule.cpp b/Engine/Source/Runtime/AudioAnalyzer/Private/AudioAnalyzerModule.cpp
--- a/Engine/Source/Runtime/AudioAnalyzer/Private/AudioAnalyzerModule.cpp
+++ b/Engine/Source/Runtime/AudioAnalyzer/Private/AudioAnalyzerModule.cpp
@@ -2,8 +2,15 @@
 
 #include "AudioAnalyzerModule.h"
 
+#include "HAL/LowLevelMemStats.h"
+#include "HAL/LowLevelMemTracker.h"
+#include "Logging/LogMacros.h"
+
 DEFINE_LOG_CATEGORY(LogAudioAnalyzer);
 
+DECLARE_LLM_MEMORY_STAT(TEXT("AudioAnalysis"), STAT_AudioAnalysisLLM, STATGROUP_LLMFULL);
+LLM_DEFINE_TAG(AudioAnalysis, TEXT("Audio Analysis"), TEXT("Audio"), GET_STATFNAME(STAT_AudioAnalysisLLM), GET_STATFNAME(STAT_AudioSummaryLLM));
+
 void FAudioAnalyzerModule::StartupModule()
 {
 }
diff --git a/Engine/Source/Runtime/AudioAnalyzer/Private/AudioAnalyzerNRT.cpp b/Engine/Source/Runtime/AudioAnalyzer/Private/AudioAnalyzerNRT.cpp
--- a/Engine/Source/Runtime/AudioAnalyzer/Private/AudioAnalyzerNRT.cpp
+++ b/Engine/Source/Runtime/AudioAnalyzer/Private/AudioAnalyzerNRT.cpp
@@ -6,7 +6,6 @@
 #include "SampleBuffer.h"
 #include "Async/Async.h"
 
-
 #if WITH_EDITOR
 
 namespace 
@@ -123,6 +122,8 @@ bool UAudioAnalyzerNRT::ShouldEventTriggerAnalysis(struct FPropertyChangedEvent
 
 void UAudioAnalyzerNRT::AnalyzeAudio()
 {
+	LLM_SCOPE_BYTAG(AudioAnalysis);
+
 	TSharedPtr<Audio::IAnalyzerNRTResult, ESPMode::ThreadSafe> NewResult;
 
 	// Create a new result id for this result.
@@ -268,6 +269,8 @@ void UAudioAnalyzerNRT::Serialize(FArchive& Ar)
 
 TUniquePtr<Audio::IAnalyzerNRTSettings> UAudioAnalyzerNRT::GetSettings(const float InSampleRate, const int32 InNumChannels) const
 {
+	LLM_SCOPE_BYTAG(AudioAnalysis);
+
 	return MakeUnique<Audio::IAnalyzerNRTSettings>();
 }
 
@@ -309,4 +312,4 @@ void UAudioAnalyzerNRT::RemoveSettingsDelegate(UAudioAnalyzerNRTSettings* InSett
 	}
 }
 
-#endif
\ No newline at end of file
+#endif
diff --git a/Engine/Source/Runtime/AudioAnalyzer/Private/AudioAnalyzerNRTFacade.cpp b/Engine/Source/Runtime/AudioAnalyzer/Private/AudioAnalyzerNRTFacade.cpp
--- a/Engine/Source/Runtime/AudioAnalyzer/Private/AudioAnalyzerNRTFacade.cpp
+++ b/Engine/Source/Runtime/AudioAnalyzer/Private/AudioAnalyzerNRTFacade.cpp
@@ -53,6 +53,8 @@ namespace Audio
 	 */
 	TUniquePtr<IAnalyzerNRTResult> FAnalyzerNRTFacade::AnalyzePCM16Audio(const TArray<uint8>& InRawWaveData, int32 InNumChannels, float InSampleRate)
 	{
+		LLM_SCOPE_BYTAG(AudioAnalysis);
+
 		// Get analyzer factory
 		IAnalyzerNRTFactory* Factory = GetAnalyzerNRTFactory(FactoryName);
 
diff --git a/Engine/Source/Runtime/AudioAnalyzer/Private/AudioAnalyzerSubsystem.cpp b/Engine/Source/Runtime/AudioAnalyzer/Private/AudioAnalyzerSubsystem.cpp
--- a/Engine/Source/Runtime/AudioAnalyzer/Private/AudioAnalyzerSubsystem.cpp
+++ b/Engine/Source/Runtime/AudioAnalyzer/Private/AudioAnalyzerSubsystem.cpp
@@ -1,6 +1,7 @@
 // Copyright Epic Games, Inc. All Rights Reserved.
 
 #include "AudioAnalyzerSubsystem.h"
+#include "AudioAnalyzerModule.h"
 #include "Containers/Ticker.h"
 #include "Engine/World.h"
 #include "Engine/Engine.h"
@@ -30,6 +31,8 @@ UAudioAnalyzerSubsystem* UAudioAnalyzerSubsystem::Get()
 
 bool UAudioAnalyzerSubsystem::Tick(float DeltaTime)
 {
+	LLM_SCOPE_BYTAG(AudioAnalysis);
+
 	// Loop through all analyzers and if they're ready to analyze, do it
 	for (UAudioAnalyzer* Analyzer : AudioAnalyzers)
 	{
@@ -51,12 +54,16 @@ void UAudioAnalyzerSubsystem::Initialize(FSubsystemCollectionBase& Collection)
 
 void UAudioAnalyzerSubsystem::Deinitialize()
 {
+	LLM_SCOPE_BYTAG(AudioAnalysis);
+
 	// Release our audio analyzers
 	AudioAnalyzers.Reset();
 }
 
 void UAudioAnalyzerSubsystem::RegisterAudioAnalyzer(UAudioAnalyzer* InAnalyzer)
 {
+	LLM_SCOPE_BYTAG(AudioAnalysis);
+
 	if (!TickerHandle.IsValid())
 	{
 		TickerHandle = FTSTicker::GetCoreTicker().AddTicker(FTickerDelegate::CreateUObject(this, &UAudioAnalyzerSubsystem::Tick), 0.0f);
@@ -66,10 +73,12 @@ void UAudioAnalyzerSubsystem::RegisterAudioAnalyzer(UAudioAnalyzer* InAnalyzer)
 
 void UAudioAnalyzerSubsystem::UnregisterAudioAnalyzer(UAudioAnalyzer* InAnalyzer)
 {
+	LLM_SCOPE_BYTAG(AudioAnalysis);
+
 	AudioAnalyzers.Remove(InAnalyzer);
 	if (AudioAnalyzers.IsEmpty() && TickerHandle.IsValid())
 	{
 		FTSTicker::GetCoreTicker().RemoveTicker(TickerHandle);
 		TickerHandle.Reset();
 	}
-}
\ No newline at end of file
+}
diff --git a/Engine/Source/Runtime/AudioAnalyzer/Public/AudioAnalyzerModule.h b/Engine/Source/Runtime/AudioAnalyzer/Public/AudioAnalyzerModule.h
--- a/Engine/Source/Runtime/AudioAnalyzer/Public/AudioAnalyzerModule.h
+++ b/Engine/Source/Runtime/AudioAnalyzer/Public/AudioAnalyzerModule.h
@@ -2,12 +2,18 @@
 
 #pragma once
 
+#if UE_ENABLE_INCLUDE_ORDER_DEPRECATED_IN_5_1 
 #include "CoreMinimal.h"
+#endif
+
+#include "HAL/LowLevelMemTracker.h"
+#include "Logging/LogMacros.h"
 #include "Modules/ModuleManager.h"
 
 DECLARE_LOG_CATEGORY_EXTERN(LogAudioAnalyzer, Log, All);
+LLM_DECLARE_TAG_API(AudioAnalysis, AUDIOANALYZER_API);
 
-class FAudioAnalyzerModule : public IModuleInterface
+class AUDIOANALYZER_API FAudioAnalyzerModule : public IModuleInterface
 {
 	public:
 	/** IModuleInterface implementation */
diff --git a/Engine/Source/Runtime/AudioMixer/Private/AudioMixerDevice.cpp b/Engine/Source/Runtime/AudioMixer/Private/AudioMixerDevice.cpp
--- a/Engine/Source/Runtime/AudioMixer/Private/AudioMixerDevice.cpp
+++ b/Engine/Source/Runtime/AudioMixer/Private/AudioMixerDevice.cpp
@@ -200,6 +200,8 @@ namespace Audio
 
 	void FMixerDevice::OnListenerUpdated(const TArray<FListener>& InListeners)
 	{
+		LLM_SCOPE(ELLMTag::AudioMixer);
+
 		ListenerTransforms.Reset(InListeners.Num());
 
 		for (const FListener& Listener : InListeners)
