Swap set property functional iterator to GetElement to for API conformity
#rnx

diff --git a/Engine/Source/Editor/AudioEditor/Private/SoundModulationDestinationLayout.cpp b/Engine/Source/Editor/AudioEditor/Private/SoundModulationDestinationLayout.cpp
--- a/Engine/Source/Editor/AudioEditor/Private/SoundModulationDestinationLayout.cpp
+++ b/Engine/Source/Editor/AudioEditor/Private/SoundModulationDestinationLayout.cpp
@@ -381,13 +381,17 @@ namespace ModDestinationLayoutUtils
 						.Font(IDetailLayoutBuilder::GetDetailFontBold())
 						.Text(TAttribute<FText>::Create([ModSet = ModulatorsHandle->AsSet(), StructPropertyHandle]()
 						{
+							uint32 NumElements = 0;
+							ModSet->GetNumElements(NumElements);
+
 							FText MismatchedText = FText::GetEmpty();
-							ModSet->IterateElements([&MismatchedText, &StructPropertyHandle](TSharedRef<IPropertyHandle> ModulatorHandle)
+							for (int32 i = 0; i < (int32)NumElements; ++i)
 							{
+								TSharedRef<IPropertyHandle> SetMemberHandle = ModSet->GetElement(i);
 								FName ModParamName;
 								FName DestName;
 								FName ModName;
-								if (ModDestinationLayoutUtils::IsParamMismatched(ModulatorHandle, StructPropertyHandle, &ModParamName, &DestName, &ModName))
+								if (ModDestinationLayoutUtils::IsParamMismatched(SetMemberHandle, StructPropertyHandle, &ModParamName, &DestName, &ModName))
 								{
 									MismatchedText = FText::Format(
 										LOCTEXT("ModulationDestinationLayout_UnitMismatchesFormat", "{0}{1} ({2})\n"),
@@ -396,7 +400,7 @@ namespace ModDestinationLayoutUtils
 										FText::FromName(ModParamName)
 									);
 								}
-							});
+							}
 
 							return MismatchedText;
 						}))
@@ -410,14 +414,19 @@ namespace ModDestinationLayoutUtils
 					return Visibility;
 				}
 
+				uint32 NumElements = 0;
+				ModSet->GetNumElements(NumElements);
+
 				bool bIsMismatched = false;
-				ModSet->IterateElements([&bIsMismatched, &StructPropertyHandle](TSharedPtr<IPropertyHandle> ModulatorHandle)
+				for (int32 i = 0; i < (int32)NumElements; ++i)
 				{
-					if (ModDestinationLayoutUtils::IsParamMismatched(ModulatorHandle->AsShared(), StructPropertyHandle))
+					TSharedRef<IPropertyHandle> SetMemberHandle = ModSet->GetElement(i);
+					if (ModDestinationLayoutUtils::IsParamMismatched(SetMemberHandle, StructPropertyHandle))
 					{
 						bIsMismatched = true;
 					}
-				});
+				}
+
 				return bIsMismatched ? EVisibility::Visible : EVisibility::Hidden;
 			}));
 	}
diff --git a/Engine/Source/Editor/PropertyEditor/Private/PropertyHandleImpl.cpp b/Engine/Source/Editor/PropertyEditor/Private/PropertyHandleImpl.cpp
--- a/Engine/Source/Editor/PropertyEditor/Private/PropertyHandleImpl.cpp
+++ b/Engine/Source/Editor/PropertyEditor/Private/PropertyHandleImpl.cpp
@@ -30,7 +30,6 @@
 #include "PropertyPathHelpers.h"
 #include "PropertyTextUtilities.h"
 #include "HAL/PlatformApplicationMisc.h"
-#include "Templates/Function.h"
 
 #define LOCTEXT_NAMESPACE "PropertyHandleImplementation"
 
@@ -4862,17 +4861,10 @@ FPropertyAccess::Result FPropertyHandleSet::GetNumElements(uint32& OutNumChildre
 	return FPropertyAccess::Success;
 }
 
-void FPropertyHandleSet::IterateElements(TFunctionRef<void(TSharedRef<IPropertyHandle>)> Func)
+TSharedRef<IPropertyHandle> FPropertyHandleSet::GetElement(int32 Index) const
 {
-	uint32 NumChildren = 0;
-	GetNumChildren(NumChildren);
-
-	for (uint32 Index = 0; Index < NumChildren; ++Index)
-	{
-		TSharedPtr<FPropertyNode> PropertyNode = Implementation->GetChildNode(Index);
-		TSharedRef<IPropertyHandle> ChildHandle = PropertyEditorHelpers::GetPropertyHandle(PropertyNode.ToSharedRef(), Implementation->GetNotifyHook(), Implementation->GetPropertyUtilities()).ToSharedRef();
-		Func(ChildHandle);
-	}
+	TSharedPtr<FPropertyNode> PropertyNode = Implementation->GetChildNode(Index);
+	return PropertyEditorHelpers::GetPropertyHandle(PropertyNode.ToSharedRef(), Implementation->GetNotifyHook(), Implementation->GetPropertyUtilities()).ToSharedRef();
 }
 
 void FPropertyHandleSet::SetOnNumElementsChanged(FSimpleDelegate& OnChildrenChanged)
diff --git a/Engine/Source/Editor/PropertyEditor/Private/PropertyHandleImpl.h b/Engine/Source/Editor/PropertyEditor/Private/PropertyHandleImpl.h
--- a/Engine/Source/Editor/PropertyEditor/Private/PropertyHandleImpl.h
+++ b/Engine/Source/Editor/PropertyEditor/Private/PropertyHandleImpl.h
@@ -706,7 +706,7 @@ public:
 	virtual FPropertyAccess::Result Empty() override;
 	virtual FPropertyAccess::Result DeleteItem(int32 Index) override;
 	virtual FPropertyAccess::Result GetNumElements(uint32& OutNumElements) override;
-	virtual void IterateElements(TFunctionRef<void(TSharedRef<IPropertyHandle>)> Func) override;
+	virtual TSharedRef<IPropertyHandle> GetElement(int32 Index) const override;
 
 	virtual void SetOnNumElementsChanged(FSimpleDelegate& InOnNumElementsChanged) override;
 	virtual bool HasDocumentation() override { return true; }
diff --git a/Engine/Source/Editor/PropertyEditor/Public/PropertyHandle.h b/Engine/Source/Editor/PropertyEditor/Public/PropertyHandle.h
--- a/Engine/Source/Editor/PropertyEditor/Public/PropertyHandle.h
+++ b/Engine/Source/Editor/PropertyEditor/Public/PropertyHandle.h
@@ -7,7 +7,6 @@
 #include "Widgets/SWidget.h"
 #include "PropertyEditorModule.h"
 #include "UObject/PropertyPortFlags.h"
-#include "Templates/Function.h"
 
 struct FAssetData;
 class IPropertyHandleArray;
@@ -743,7 +742,7 @@ public:
 	virtual FPropertyAccess::Result GetNumElements( uint32& OutNumItems ) const = 0;
 
 	/**
-	 * @return a handle to the element at the specified index                                                              
+	 * @return a handle to the element at the specified index
 	 */
 	virtual TSharedRef<IPropertyHandle> GetElement( int32 Index ) const = 0;
 
@@ -755,7 +754,7 @@ public:
 
 
 	/**
-	 * Sets a delegate to call when the number of elements changes                                                  
+	 * Sets a delegate to call when the number of elements changes
 	 */
 	virtual void SetOnNumElementsChanged( FSimpleDelegate& InOnNumElementsChanged ) = 0;
 };
@@ -797,9 +796,9 @@ public:
 	virtual FPropertyAccess::Result GetNumElements(uint32& OutNumElements) = 0;
 
 	/**
-	 * Iterates over all set elements using the provided function
+	 * @return a handle to the element at the specified index
 	 */
-	virtual void IterateElements(TFunctionRef<void(TSharedRef<IPropertyHandle>)> Func) = 0;
+	virtual TSharedRef<IPropertyHandle> GetElement(int32 Index) const = 0;
 
 
 	/**
