Fix Audio Modulation Parameter Mismatch UI Bug

diff --git a/Engine/Plugins/Runtime/AudioModulation/Source/AudioModulation/Private/SoundModulationParameter.cpp b/Engine/Plugins/Runtime/AudioModulation/Source/AudioModulation/Private/SoundModulationParameter.cpp
--- a/Engine/Plugins/Runtime/AudioModulation/Source/AudioModulation/Private/SoundModulationParameter.cpp
+++ b/Engine/Plugins/Runtime/AudioModulation/Source/AudioModulation/Private/SoundModulationParameter.cpp
@@ -66,6 +66,7 @@ Audio::FModulationParameter USoundModulationParameter::CreateParameter() const
 
 #if WITH_EDITORONLY_DATA
 	Parameter.UnitDisplayName = Settings.UnitDisplayName;
+	Parameter.ClassName = GetClass()->GetFName();
 #endif // WITH_EDITORONLY_DATA
 
 	return Parameter;
diff --git a/Engine/Plugins/Runtime/Synthesis/Source/Synthesis/Classes/SourceEffects/SourceEffectBitCrusher.h b/Engine/Plugins/Runtime/Synthesis/Source/Synthesis/Classes/SourceEffects/SourceEffectBitCrusher.h
--- a/Engine/Plugins/Runtime/Synthesis/Source/Synthesis/Classes/SourceEffects/SourceEffectBitCrusher.h
+++ b/Engine/Plugins/Runtime/Synthesis/Source/Synthesis/Classes/SourceEffects/SourceEffectBitCrusher.h
@@ -35,14 +35,14 @@ struct SYNTHESIS_API FSourceEffectBitCrusherSettings
 	float CrushedSampleRate = 8000.0f;
 
 	// The reduced frequency to use for the audio stream. 
-	UPROPERTY(EditAnywhere, BlueprintReadWrite, Category = "SourceEffect|Preset", meta = (DisplayName = "Sample Rate", AudioParam = "SampleRate", ClampMin = "0.1", ClampMax = "96000.0", UIMin = "500.0", UIMax = "16000.0"))
+	UPROPERTY(EditAnywhere, BlueprintReadWrite, Category = "SourceEffect|Preset", meta = (DisplayName = "Sample Rate", AudioParam = "SampleRate", AudioParamClass = "SoundModulationParameterFrequency", ClampMin = "0.1", ClampMax = "96000.0", UIMin = "500.0", UIMax = "16000.0"))
 	FSoundModulationDestinationSettings SampleRateModulation;
 
 	UPROPERTY(meta = (PropertyDeprecated))
 	float CrushedBits = 8.0f;
 
 	// The reduced bit depth to use for the audio stream
-	UPROPERTY(EditAnywhere, BlueprintReadWrite, Category = "SourceEffect|Preset", meta = (DisplayName = "Bit Depth", AudioParam = "BitDepth", ClampMin = "1.0", ClampMax = "24.0", UIMin = "1.0", UIMax = "16.0"))
+	UPROPERTY(EditAnywhere, BlueprintReadWrite, Category = "SourceEffect|Preset", meta = (DisplayName = "Bit Depth", AudioParam = "BitDepth", AudioParamClass = "SoundModulationParameterScaled", ClampMin = "1.0", ClampMax = "24.0", UIMin = "1.0", UIMax = "16.0"))
 	FSoundModulationDestinationSettings BitModulation;
 
 	FSourceEffectBitCrusherSettings()
diff --git a/Engine/Plugins/Runtime/Synthesis/Source/Synthesis/Classes/SourceEffects/SourceEffectChorus.h b/Engine/Plugins/Runtime/Synthesis/Source/Synthesis/Classes/SourceEffects/SourceEffectChorus.h
--- a/Engine/Plugins/Runtime/Synthesis/Source/Synthesis/Classes/SourceEffects/SourceEffectChorus.h
+++ b/Engine/Plugins/Runtime/Synthesis/Source/Synthesis/Classes/SourceEffects/SourceEffectChorus.h
@@ -63,27 +63,27 @@ struct SYNTHESIS_API FSourceEffectChorusSettings
 	float Spread = 0.0f;
 
 	// The depth of the chorus effect
-	UPROPERTY(EditAnywhere, BlueprintReadWrite, Category = "SourceEffect|Preset", meta = (DisplayName = "Depth", ClampMin = "0.0", ClampMax = "1.0", UIMin = "0.0", UIMax = "1.0"))
+	UPROPERTY(EditAnywhere, BlueprintReadWrite, Category = "SourceEffect|Preset", meta = (DisplayName = "Depth", AudioParamClass = "SoundModulationParameter", ClampMin = "0.0", ClampMax = "1.0", UIMin = "0.0", UIMax = "1.0"))
 	FSoundModulationDestinationSettings DepthModulation;
 
 	// The frequency of the chorus effect
-	UPROPERTY(EditAnywhere, BlueprintReadWrite, Category = "SourceEffect|Preset", meta = (DisplayName = "Frequency", AudioParam = "LowRateFrequency", UIMin = "0.0", UIMax = "5.0"))
+	UPROPERTY(EditAnywhere, BlueprintReadWrite, Category = "SourceEffect|Preset", meta = (DisplayName = "Frequency", AudioParam = "LowRateFrequency", AudioParamClass = "SoundModulationParameterFrequency", UIMin = "0.0", UIMax = "5.0"))
 	FSoundModulationDestinationSettings FrequencyModulation;
 
 	// The feedback of the chorus effect
-	UPROPERTY(EditAnywhere, BlueprintReadWrite, Category = "SourceEffect|Preset", meta = (DisplayName = "Feedback", ClampMin = "0.0", ClampMax = "1.0", UIMin = "0.0", UIMax = "1.0"))
+	UPROPERTY(EditAnywhere, BlueprintReadWrite, Category = "SourceEffect|Preset", meta = (DisplayName = "Feedback", AudioParamClass = "SoundModulationParameter", ClampMin = "0.0", ClampMax = "1.0", UIMin = "0.0", UIMax = "1.0"))
 	FSoundModulationDestinationSettings FeedbackModulation;
 
 	// The wet level of the chorus effect
-	UPROPERTY(EditAnywhere, BlueprintReadWrite, Category = "SourceEffect|Preset", meta = (DisplayName = "Wet Level", ClampMin = "0.0", ClampMax = "1.0", UIMin = "0.0", UIMax = "1.0"))
+	UPROPERTY(EditAnywhere, BlueprintReadWrite, Category = "SourceEffect|Preset", meta = (DisplayName = "Wet Level", AudioParamClass = "SoundModulationParameter", ClampMin = "0.0", ClampMax = "1.0", UIMin = "0.0", UIMax = "1.0"))
 	FSoundModulationDestinationSettings WetModulation;
 
 	// The dry level of the chorus effect
-	UPROPERTY(EditAnywhere, BlueprintReadWrite, Category = "SourceEffect|Preset", meta = (DisplayName = "Dry Level", ClampMin = "0.0", ClampMax = "1.0", UIMin = "0.0", UIMax = "1.0"))
+	UPROPERTY(EditAnywhere, BlueprintReadWrite, Category = "SourceEffect|Preset", meta = (DisplayName = "Dry Level", AudioParamClass = "SoundModulationParameter", ClampMin = "0.0", ClampMax = "1.0", UIMin = "0.0", UIMax = "1.0"))
 	FSoundModulationDestinationSettings DryModulation;
 
 	// The spread of the effect (larger means greater difference between left and right delay lines)
-	UPROPERTY(EditAnywhere, BlueprintReadWrite, Category = "SourceEffect|Preset", meta = (DisplayName = "Spread", ClampMin = "0.0", ClampMax = "1.0", UIMin = "0.0", UIMax = "1.0"))
+	UPROPERTY(EditAnywhere, BlueprintReadWrite, Category = "SourceEffect|Preset", meta = (DisplayName = "Spread", AudioParamClass = "SoundModulationParameter", ClampMin = "0.0", ClampMax = "1.0", UIMin = "0.0", UIMax = "1.0"))
 	FSoundModulationDestinationSettings SpreadModulation;
 
 	FSourceEffectChorusSettings()
diff --git a/Engine/Source/Editor/AudioEditor/Private/SoundModulationDestinationLayout.cpp b/Engine/Source/Editor/AudioEditor/Private/SoundModulationDestinationLayout.cpp
--- a/Engine/Source/Editor/AudioEditor/Private/SoundModulationDestinationLayout.cpp
+++ b/Engine/Source/Editor/AudioEditor/Private/SoundModulationDestinationLayout.cpp
@@ -88,16 +88,29 @@ namespace ModDestinationLayoutUtils
 		return FName();
 	}
 
-	bool IsParamMismatched(TSharedRef<IPropertyHandle> ModulatorHandle, TSharedRef<IPropertyHandle> StructPropertyHandle, FName* OutModParamName = nullptr, FName* OutDestParamName = nullptr, FName* OutModName = nullptr)
+	FName GetParameterClassFromMetaData(const TSharedRef<IPropertyHandle>& InHandle)
 	{
-		if (OutModParamName)
+		static const FName AudioParamClassFieldName("AudioParamClass");
+
+		if (InHandle->HasMetaData(AudioParamClassFieldName))
+		{
+			const FName Param = *(InHandle->GetMetaData(AudioParamClassFieldName));
+			return Param;
+		}
+
+		return FName();
+	}
+
+	bool IsParamMismatched(TSharedRef<IPropertyHandle> ModulatorHandle, TSharedRef<IPropertyHandle> StructPropertyHandle, FName* OutModParamClassName = nullptr, FName* OutDestParamClassName = nullptr, FName* OutModName = nullptr)
+	{
+		if (OutModParamClassName)
 		{
-			*OutModParamName = FName();
+			*OutModParamClassName = FName();
 		}
 
-		if (OutDestParamName)
+		if (OutDestParamClassName)
 		{
-			*OutDestParamName = FName();
+			*OutDestParamClassName = FName();
 		}
 
 		if (OutModName)
@@ -114,18 +127,18 @@ namespace ModDestinationLayoutUtils
 			return false;
 		}
 
-		const FName ModParamName = ModBase->GetOutputParameter().ParameterName;
-		const FName DestParamName = ModDestinationLayoutUtils::GetParameterNameFromMetaData(StructPropertyHandle);
-		if (ModParamName != FName() && DestParamName != FName() && ModParamName != DestParamName)
+		const FName ModParamClassName = ModBase->GetOutputParameter().ClassName;
+		const FName DestParamClassName = ModDestinationLayoutUtils::GetParameterClassFromMetaData(StructPropertyHandle);
+		if (!ModParamClassName.IsNone() && !DestParamClassName.IsNone() && ModParamClassName != DestParamClassName)
 		{
-			if (OutModParamName)
+			if (OutModParamClassName)
 			{
-				*OutModParamName = ModParamName;
+				*OutModParamClassName = ModParamClassName;
 			}
 
-			if (OutDestParamName)
+			if (OutDestParamClassName)
 			{
-				*OutDestParamName = DestParamName;
+				*OutDestParamClassName = DestParamClassName;
 			}
 
 			if (OutModName)
@@ -384,25 +397,25 @@ namespace ModDestinationLayoutUtils
 							uint32 NumElements = 0;
 							ModSet->GetNumElements(NumElements);
 
-							FText MismatchedText = FText::GetEmpty();
+							FString MismatchedText("");
 							for (int32 i = 0; i < (int32)NumElements; ++i)
 							{
 								TSharedRef<IPropertyHandle> SetMemberHandle = ModSet->GetElement(i);
-								FName ModParamName;
-								FName DestName;
+								FName ModParamClassName;
+								FName DestClassName;
 								FName ModName;
-								if (ModDestinationLayoutUtils::IsParamMismatched(SetMemberHandle, StructPropertyHandle, &ModParamName, &DestName, &ModName))
+								if (ModDestinationLayoutUtils::IsParamMismatched(SetMemberHandle, StructPropertyHandle, &ModParamClassName, &DestClassName, &ModName))
 								{
-									MismatchedText = FText::Format(
-										LOCTEXT("ModulationDestinationLayout_UnitMismatchesFormat", "{0}{1} ({2})\n"),
-										MismatchedText,
+									MismatchedText += FText::Format(
+										LOCTEXT("ModulationDestinationLayout_UnitMismatchesFormat", "{0} ({1}), Expected: {2}\n"),
 										FText::FromName(ModName),
-										FText::FromName(ModParamName)
-									);
+										FText::FromName(ModParamClassName),
+										FText::FromName(DestClassName)
+									).ToString();
 								}
 							}
 
-							return MismatchedText;
+							return FText::FromString(MismatchedText);
 						}))
 					]
 			]
diff --git a/Engine/Source/Runtime/AudioExtensions/Private/IAudioModulation.cpp b/Engine/Source/Runtime/AudioExtensions/Private/IAudioModulation.cpp
--- a/Engine/Source/Runtime/AudioExtensions/Private/IAudioModulation.cpp
+++ b/Engine/Source/Runtime/AudioExtensions/Private/IAudioModulation.cpp
@@ -85,6 +85,7 @@ namespace Audio
 		, bRequiresConversion(InParam.bRequiresConversion)
 	#if WITH_EDITORONLY_DATA
 		, UnitDisplayName(MoveTemp(InParam.UnitDisplayName))
+		, ClassName(MoveTemp(InParam.ClassName))
 	#endif // WITH_EDITORONLY_DATA
 		, MixFunction(MoveTemp(InParam.MixFunction))
 		, UnitFunction(MoveTemp(InParam.UnitFunction))
@@ -100,6 +101,7 @@ namespace Audio
 		, bRequiresConversion(InParam.bRequiresConversion)
 #if WITH_EDITORONLY_DATA
 		, UnitDisplayName(InParam.UnitDisplayName)
+		, ClassName(InParam.ClassName)
 #endif // WITH_EDITORONLY_DATA
 		, MixFunction(InParam.MixFunction)
 		, UnitFunction(InParam.UnitFunction)
@@ -117,6 +119,7 @@ namespace Audio
 
 #if WITH_EDITORONLY_DATA
 		UnitDisplayName = InParam.UnitDisplayName;
+		ClassName = InParam.ClassName;
 #endif // WITH_EDITORONLY_DATA
 
 		MixFunction = InParam.MixFunction;
@@ -136,6 +139,7 @@ namespace Audio
 
 	#if WITH_EDITORONLY_DATA
 		UnitDisplayName = MoveTemp(InParam.UnitDisplayName);
+		ClassName = MoveTemp(InParam.ClassName);
 	#endif // WITH_EDITORONLY_DATA
 
 		MixFunction = MoveTemp(InParam.MixFunction);
diff --git a/Engine/Source/Runtime/AudioExtensions/Public/IAudioModulation.h b/Engine/Source/Runtime/AudioExtensions/Public/IAudioModulation.h
--- a/Engine/Source/Runtime/AudioExtensions/Public/IAudioModulation.h
+++ b/Engine/Source/Runtime/AudioExtensions/Public/IAudioModulation.h
@@ -73,6 +73,8 @@ namespace Audio
 
 #if WITH_EDITORONLY_DATA
 		FText UnitDisplayName;
+
+		FName ClassName;
 #endif // WITH_EDITORONLY_DATA
 
 		// Function used to mix normalized values together.
diff --git a/Engine/Source/Runtime/Engine/Classes/Sound/SoundModulationDestination.h b/Engine/Source/Runtime/Engine/Classes/Sound/SoundModulationDestination.h
--- a/Engine/Source/Runtime/Engine/Classes/Sound/SoundModulationDestination.h
+++ b/Engine/Source/Runtime/Engine/Classes/Sound/SoundModulationDestination.h
@@ -78,19 +78,19 @@ struct ENGINE_API FSoundModulationDefaultSettings
 	FSoundModulationDefaultSettings();
 
 	/** Volume modulation */
-	UPROPERTY(EditAnywhere, BlueprintReadWrite, Category = "Modulation", meta = (DisplayName = "Volume", AudioParam = "Volume"))
+	UPROPERTY(EditAnywhere, BlueprintReadWrite, Category = "Modulation", meta = (DisplayName = "Volume", AudioParam = "Volume", AudioParamClass = "SoundModulationParameterVolume"))
 	FSoundModulationDestinationSettings VolumeModulationDestination;
 
 	/** Pitch modulation */
-	UPROPERTY(EditAnywhere, BlueprintReadWrite, Category = "Modulation", meta = (DisplayName = "Pitch", AudioParam = "Pitch"))
+	UPROPERTY(EditAnywhere, BlueprintReadWrite, Category = "Modulation", meta = (DisplayName = "Pitch", AudioParam = "Pitch", AudioParamClass = "SoundModulationParameterBipolar"))
 	FSoundModulationDestinationSettings PitchModulationDestination;
 
 	/** Highpass modulation */
-	UPROPERTY(EditAnywhere, BlueprintReadWrite, Category = "Modulation", meta = (DisplayName = "Highpass", AudioParam = "HPFCutoffFrequency"))
+	UPROPERTY(EditAnywhere, BlueprintReadWrite, Category = "Modulation", meta = (DisplayName = "Highpass", AudioParam = "HPFCutoffFrequency", AudioParamClass = "SoundModulationParameterHPFFrequency"))
 	FSoundModulationDestinationSettings HighpassModulationDestination;
 
 	/** Lowpass modulation */
-	UPROPERTY(EditAnywhere, BlueprintReadWrite, Category = "Modulation", meta = (DisplayName = "Lowpass", AudioParam = "LPFCutoffFrequency"))
+	UPROPERTY(EditAnywhere, BlueprintReadWrite, Category = "Modulation", meta = (DisplayName = "Lowpass", AudioParam = "LPFCutoffFrequency", AudioParamClass = "SoundModulationParameterLPFFrequency"))
 	FSoundModulationDestinationSettings LowpassModulationDestination;
 
 #if WITH_EDITORONLY_DATA
diff --git a/Engine/Source/Runtime/Engine/Classes/Sound/SoundSubmix.h b/Engine/Source/Runtime/Engine/Classes/Sound/SoundSubmix.h
--- a/Engine/Source/Runtime/Engine/Classes/Sound/SoundSubmix.h
+++ b/Engine/Source/Runtime/Engine/Classes/Sound/SoundSubmix.h
@@ -296,15 +296,15 @@ public:
 	float DryLevel;
 
 	/** The output volume of the submix in Decibels. Applied after submix effects and analysis are performed.*/
-	UPROPERTY(EditAnywhere, BlueprintReadWrite, Category = SubmixLevel, meta = (DisplayName = "Output Volume (dB)", AudioParam = "Volume", ClampMin = "-96.0", ClampMax = "0.0", UIMin = "-96.0", UIMax = "0.0"))
+	UPROPERTY(EditAnywhere, BlueprintReadWrite, Category = SubmixLevel, meta = (DisplayName = "Output Volume (dB)", AudioParam = "Volume", AudioParamClass = "SoundModulationParameterVolume", ClampMin = "-96.0", ClampMax = "0.0", UIMin = "-96.0", UIMax = "0.0"))
 	FSoundModulationDestinationSettings OutputVolumeModulation;
 
 	/** The wet level of the submixin Decibels. Applied after submix effects and analysis are performed. */
-	UPROPERTY(EditAnywhere, BlueprintReadWrite, Category = SubmixLevel, meta = (DisplayName = "Wet Level (dB)", AudioParam = "Volume", ClampMin = "-96.0", ClampMax = "0.0", UIMin = "-96.0", UIMax = "0.0"))
+	UPROPERTY(EditAnywhere, BlueprintReadWrite, Category = SubmixLevel, meta = (DisplayName = "Wet Level (dB)", AudioParam = "Volume", AudioParamClass = "SoundModulationParameterVolume", ClampMin = "-96.0", ClampMax = "0.0", UIMin = "-96.0", UIMax = "0.0"))
 	FSoundModulationDestinationSettings WetLevelModulation;
 
 	/** The dry level of the submix in Decibels. Applied before submix effects and analysis are performed. */
-	UPROPERTY(EditAnywhere, BlueprintReadWrite, Category = SubmixLevel, meta = (DisplayName = "Dry Level (dB)", AudioParam = "Volume", ClampMin = "-96.0", ClampMax = "0.0", UIMin = "-96.0", UIMax = "0.0"))
+	UPROPERTY(EditAnywhere, BlueprintReadWrite, Category = SubmixLevel, meta = (DisplayName = "Dry Level (dB)", AudioParam = "Volume", AudioParamClass = "SoundModulationParameterVolume", ClampMin = "-96.0", ClampMax = "0.0", UIMin = "-96.0", UIMax = "0.0"))
 	FSoundModulationDestinationSettings DryLevelModulation;
 	
 	/** Whether to send this Submix to AudioLink (when AudioLink is Enabled)*/
