AssetDefinition - Upgrading several sound asset actions.

diff --git a/Engine/Source/Editor/AssetDefinition/Public/AssetDefinition.h b/Engine/Source/Editor/AssetDefinition/Public/AssetDefinition.h
--- a/Engine/Source/Editor/AssetDefinition/Public/AssetDefinition.h
+++ b/Engine/Source/Editor/AssetDefinition/Public/AssetDefinition.h
@@ -77,6 +77,23 @@ struct FAssetArgs
 		
 		return LoadedObjects;
 	}
+	
+	template<typename ExpectedObjectType>
+    ExpectedObjectType* LoadFirstValid(const TSet<FName>& LoadTags = {}) const
+    {   	
+    	for (const FAssetData& Asset : Assets)
+    	{
+    		if (Asset.IsInstanceOf(ExpectedObjectType::StaticClass()))
+    		{
+    			if (ExpectedObjectType* ExpectedType = Cast<ExpectedObjectType>(Asset.GetAsset(LoadTags)))
+    			{
+    				return ExpectedType;
+    			}
+    		}	
+    	}
+    	
+    	return nullptr;
+    }
 };
 
 struct FAssetOpenArgs : public FAssetArgs
diff --git a/Engine/Source/Editor/AudioEditor/Private/AssetTypeActions/AssetDefinition_AudioBus.h b/Engine/Source/Editor/AudioEditor/Private/AssetTypeActions/AssetDefinition_AudioBus.h
new file mode 100644
--- /dev/null
+++ b/Engine/Source/Editor/AudioEditor/Private/AssetTypeActions/AssetDefinition_AudioBus.h
@@ -0,0 +1,27 @@
+// Copyright Epic Games, Inc. All Rights Reserved.
+
+#pragma once
+
+#include "CoreMinimal.h"
+#include "Sound/AudioBus.h"
+#include "AssetDefinitionDefault.h"
+
+#include "AssetDefinition_AudioBus.generated.h"
+
+UCLASS()
+class UAssetDefinition_AudioBus : public UAssetDefinitionDefault
+{
+	GENERATED_BODY()
+
+public:
+	// UAssetDefinition Begin
+	virtual FText GetAssetDisplayName() const override { return NSLOCTEXT("AssetTypeActions", "AssetTypeActions_AudioBus", "Audio Bus"); }
+	virtual FLinearColor GetAssetColor() const override { return FLinearColor(FColor(97, 97, 85)); }
+	virtual TSoftClassPtr<UObject> GetAssetClass() const override { return UAudioBus::StaticClass(); }
+	virtual TConstArrayView<FAssetCategoryPath> GetAssetCategories() const override
+	{
+		static const auto Categories = { FAssetCategoryPath(EAssetCategoryPaths::Audio, NSLOCTEXT("AssetTypeActions", "AssetAudioBusMenu", "Mix")) };
+		return Categories;
+	}
+	// UAssetDefinition End
+};
diff --git a/Engine/Source/Editor/AudioEditor/Private/AssetTypeActions/AssetDefinition_DialogueVoice.h b/Engine/Source/Editor/AudioEditor/Private/AssetTypeActions/AssetDefinition_DialogueVoice.h
new file mode 100644
--- /dev/null
+++ b/Engine/Source/Editor/AudioEditor/Private/AssetTypeActions/AssetDefinition_DialogueVoice.h
@@ -0,0 +1,27 @@
+// Copyright Epic Games, Inc. All Rights Reserved.
+
+#pragma once
+
+#include "CoreMinimal.h"
+#include "Sound/DialogueVoice.h"
+#include "AssetDefinitionDefault.h"
+
+#include "AssetDefinition_DialogueVoice.generated.h"
+
+UCLASS()
+class UAssetDefinition_DialogueVoice : public UAssetDefinitionDefault
+{
+	GENERATED_BODY()
+
+public:
+	// UAssetDefinition Begin
+	virtual FText GetAssetDisplayName() const override { return NSLOCTEXT("AssetTypeActions", "AssetTypeActions_DialogueVoice", "Dialogue Voice"); }
+	virtual FLinearColor GetAssetColor() const override { return FLinearColor(FColor(97, 85, 212)); }
+	virtual TSoftClassPtr<UObject> GetAssetClass() const override { return UDialogueVoice::StaticClass(); }
+	virtual TConstArrayView<FAssetCategoryPath> GetAssetCategories() const override
+	{
+		static const auto Categories = { FAssetCategoryPath(EAssetCategoryPaths::Audio, NSLOCTEXT("AssetTypeActions", "AssetDialogueSubMenu", "Dialogue")) };
+		return Categories;
+	}
+	// UAssetDefinition End
+};
diff --git a/Engine/Source/Editor/AudioEditor/Private/AssetTypeActions/AssetDefinition_SoundAttenuation.h b/Engine/Source/Editor/AudioEditor/Private/AssetTypeActions/AssetDefinition_SoundAttenuation.h
new file mode 100644
--- /dev/null
+++ b/Engine/Source/Editor/AudioEditor/Private/AssetTypeActions/AssetDefinition_SoundAttenuation.h
@@ -0,0 +1,27 @@
+// Copyright Epic Games, Inc. All Rights Reserved.
+
+#pragma once
+
+#include "CoreMinimal.h"
+#include "Sound/SoundAttenuation.h"
+#include "AssetDefinitionDefault.h"
+
+#include "AssetDefinition_SoundAttenuation.generated.h"
+
+UCLASS()
+class UAssetDefinition_SoundAttenuation : public UAssetDefinitionDefault
+{
+	GENERATED_BODY()
+
+public:
+	// UAssetDefinition Begin
+	virtual FText GetAssetDisplayName() const override { return NSLOCTEXT("AssetTypeActions", "AssetTypeActions_SoundAttenuation", "Sound Attenuation"); }
+	virtual FLinearColor GetAssetColor() const override { return FLinearColor(FColor(77, 120, 239)); }
+	virtual TSoftClassPtr<UObject> GetAssetClass() const override { return USoundAttenuation::StaticClass(); }
+	virtual TConstArrayView<FAssetCategoryPath> GetAssetCategories() const override
+	{
+		static const auto Categories = { FAssetCategoryPath(EAssetCategoryPaths::Audio, NSLOCTEXT("AssetTypeActions", "AssetSoundAttenuationSubMenu", "Spatialization")) };
+		return Categories;
+	}
+	// UAssetDefinition End
+};
diff --git a/Engine/Source/Editor/AudioEditor/Private/AssetTypeActions/AssetDefinition_SoundConcurrency.h b/Engine/Source/Editor/AudioEditor/Private/AssetTypeActions/AssetDefinition_SoundConcurrency.h
new file mode 100644
--- /dev/null
+++ b/Engine/Source/Editor/AudioEditor/Private/AssetTypeActions/AssetDefinition_SoundConcurrency.h
@@ -0,0 +1,27 @@
+// Copyright Epic Games, Inc. All Rights Reserved.
+
+#pragma once
+
+#include "CoreMinimal.h"
+#include "Sound/SoundConcurrency.h"
+#include "AssetDefinitionDefault.h"
+
+#include "AssetDefinition_SoundConcurrency.generated.h"
+
+UCLASS()
+class UAssetDefinition_SoundConcurrency : public UAssetDefinitionDefault
+{
+	GENERATED_BODY()
+
+public:
+	// UAssetDefinition Begin
+	virtual FText GetAssetDisplayName() const override { return NSLOCTEXT("AssetTypeActions", "AssetTypeActions_SoundConcurrency", "Sound Concurrency"); }
+	virtual FLinearColor GetAssetColor() const override { return FLinearColor(FColor(77, 100, 139)); }
+	virtual TSoftClassPtr<UObject> GetAssetClass() const override { return USoundConcurrency::StaticClass(); }
+	virtual TConstArrayView<FAssetCategoryPath> GetAssetCategories() const override
+	{
+		static const auto Categories = { FAssetCategoryPath(EAssetCategoryPaths::Audio, NSLOCTEXT("AssetTypeActions", "AssetSoundConcurrencySubMenu", "Mix")) };
+		return Categories;
+	}
+	// UAssetDefinition End
+};
diff --git a/Engine/Source/Editor/AudioEditor/Private/AssetTypeActions/AssetDefinition_SoundMix.h b/Engine/Source/Editor/AudioEditor/Private/AssetTypeActions/AssetDefinition_SoundMix.h
new file mode 100644
--- /dev/null
+++ b/Engine/Source/Editor/AudioEditor/Private/AssetTypeActions/AssetDefinition_SoundMix.h
@@ -0,0 +1,27 @@
+// Copyright Epic Games, Inc. All Rights Reserved.
+
+#pragma once
+
+#include "CoreMinimal.h"
+#include "Sound/SoundMix.h"
+#include "AssetDefinitionDefault.h"
+
+#include "AssetDefinition_SoundMix.generated.h"
+
+UCLASS()
+class UAssetDefinition_SoundMix : public UAssetDefinitionDefault
+{
+	GENERATED_BODY()
+
+public:
+	// UAssetDefinition Begin
+	virtual FText GetAssetDisplayName() const override { return NSLOCTEXT("AssetTypeActions", "AssetTypeActions_SoundMix", "Sound Class Mix"); }
+	virtual FLinearColor GetAssetColor() const override { return FLinearColor(FColor(175, 0, 255)); }
+	virtual TSoftClassPtr<UObject> GetAssetClass() const override { return USoundMix::StaticClass(); }
+	virtual TConstArrayView<FAssetCategoryPath> GetAssetCategories() const override
+	{
+		static const auto Categories = { FAssetCategoryPath(EAssetCategoryPaths::Audio, NSLOCTEXT("AssetTypeActions", "AssetSoundClassSubMenu", "Classes")) };
+		return Categories;
+	}
+	// UAssetDefinition End
+};
diff --git a/Engine/Source/Editor/AudioEditor/Private/AssetTypeActions/AssetTypeActions_AudioBus.cpp b/Engine/Source/Editor/AudioEditor/Private/AssetTypeActions/AssetTypeActions_AudioBus.cpp
deleted file mode 100644
diff --git a/Engine/Source/Editor/AudioEditor/Private/AssetTypeActions/AssetTypeActions_AudioBus.h b/Engine/Source/Editor/AudioEditor/Private/AssetTypeActions/AssetTypeActions_AudioBus.h
deleted file mode 100644
diff --git a/Engine/Source/Editor/AudioEditor/Private/AssetTypeActions/AssetTypeActions_DialogueVoice.h b/Engine/Source/Editor/AudioEditor/Private/AssetTypeActions/AssetTypeActions_DialogueVoice.h
deleted file mode 100644
diff --git a/Engine/Source/Editor/AudioEditor/Private/AssetTypeActions/AssetTypeActions_SoundAttenuation.cpp b/Engine/Source/Editor/AudioEditor/Private/AssetTypeActions/AssetTypeActions_SoundAttenuation.cpp
deleted file mode 100644
diff --git a/Engine/Source/Editor/AudioEditor/Private/AssetTypeActions/AssetTypeActions_SoundAttenuation.h b/Engine/Source/Editor/AudioEditor/Private/AssetTypeActions/AssetTypeActions_SoundAttenuation.h
deleted file mode 100644
diff --git a/Engine/Source/Editor/AudioEditor/Private/AssetTypeActions/AssetTypeActions_SoundConcurrency.cpp b/Engine/Source/Editor/AudioEditor/Private/AssetTypeActions/AssetTypeActions_SoundConcurrency.cpp
deleted file mode 100644
diff --git a/Engine/Source/Editor/AudioEditor/Private/AssetTypeActions/AssetTypeActions_SoundConcurrency.h b/Engine/Source/Editor/AudioEditor/Private/AssetTypeActions/AssetTypeActions_SoundConcurrency.h
deleted file mode 100644
diff --git a/Engine/Source/Editor/AudioEditor/Private/AssetTypeActions/AssetTypeActions_SoundMix.cpp b/Engine/Source/Editor/AudioEditor/Private/AssetTypeActions/AssetTypeActions_SoundMix.cpp
deleted file mode 100644
diff --git a/Engine/Source/Editor/AudioEditor/Private/AssetTypeActions/AssetTypeActions_SoundMix.h b/Engine/Source/Editor/AudioEditor/Private/AssetTypeActions/AssetTypeActions_SoundMix.h
deleted file mode 100644
diff --git a/Engine/Source/Editor/AudioEditor/Private/AudioEditorModule.cpp b/Engine/Source/Editor/AudioEditor/Private/AudioEditorModule.cpp
--- a/Engine/Source/Editor/AudioEditor/Private/AudioEditorModule.cpp
+++ b/Engine/Source/Editor/AudioEditor/Private/AudioEditorModule.cpp
@@ -3,16 +3,11 @@
 #include "AudioEditorModule.h"
 
 #include "AssetToolsModule.h"
-#include "AssetTypeActions/AssetTypeActions_AudioBus.h"
-#include "AssetTypeActions/AssetTypeActions_DialogueVoice.h"
 #include "AssetTypeActions/AssetTypeActions_DialogueWave.h"
-#include "AssetTypeActions/AssetTypeActions_SoundAttenuation.h"
 #include "AssetTypeActions/AssetTypeActions_SoundBase.h"
 #include "AssetTypeActions/AssetTypeActions_SoundClass.h"
-#include "AssetTypeActions/AssetTypeActions_SoundConcurrency.h"
 #include "AssetTypeActions/AssetTypeActions_SoundCue.h"
 #include "AssetTypeActions/AssetTypeActions_SoundEffectPreset.h"
-#include "AssetTypeActions/AssetTypeActions_SoundMix.h"
 #include "AssetTypeActions/AssetTypeActions_SoundWave.h"
 #include "AssetTypeActions/AssetTypeActions_SoundSourceBus.h"
 #include "AssetTypeActions/AssetTypeActions_SoundSubmix.h"
@@ -185,14 +180,10 @@ public:
 		// Register the audio editor asset type actions
 		IAssetTools& AssetTools = FModuleManager::LoadModuleChecked<FAssetToolsModule>("AssetTools").Get();
 
-		AssetTools.RegisterAssetTypeActions(MakeShared<FAssetTypeActions_DialogueVoice>());
 		AssetTools.RegisterAssetTypeActions(MakeShared<FAssetTypeActions_DialogueWave>());
-		AssetTools.RegisterAssetTypeActions(MakeShared<FAssetTypeActions_SoundAttenuation>());
-		AssetTools.RegisterAssetTypeActions(MakeShared<FAssetTypeActions_SoundConcurrency>());
 		AssetTools.RegisterAssetTypeActions(MakeShared<FAssetTypeActions_SoundBase>());
 		AssetTools.RegisterAssetTypeActions(MakeShared<FAssetTypeActions_SoundClass>());
 		AssetTools.RegisterAssetTypeActions(MakeShared<FAssetTypeActions_SoundCue>());
-		AssetTools.RegisterAssetTypeActions(MakeShared<FAssetTypeActions_SoundMix>());
 		AssetTools.RegisterAssetTypeActions(MakeShared<FAssetTypeActions_SoundWave>());
 	}
 
@@ -215,7 +206,6 @@ public:
 			AssetTools.RegisterAssetTypeActions(MakeShared<FAssetTypeActions_SoundEffectSourcePreset>());
 			AssetTools.RegisterAssetTypeActions(MakeShared<FAssetTypeActions_SoundEffectSourcePresetChain>());
 			AssetTools.RegisterAssetTypeActions(MakeShared<FAssetTypeActions_SoundSourceBus>());
-			AssetTools.RegisterAssetTypeActions(MakeShared<FAssetTypeActions_AudioBus>());
 		}
 	}
 
