AssetDefinition - Upgrading several sound asset actions.
AssetDefinition - Fixing a regression where we needed to load assets on right click to check a 'can we create a blueprint from this asset' option.
AssetDefinition - Making the dynamic entries for GetAssetActions nameless so that they don't stomp eachother on accident.
AssetDefinition -

diff --git a/Engine/Plugins/Editor/ContentBrowser/ContentBrowserAssetDataSource/Source/ContentBrowserAssetDataSource/Private/AssetFileContextMenu.cpp b/Engine/Plugins/Editor/ContentBrowser/ContentBrowserAssetDataSource/Source/ContentBrowserAssetDataSource/Private/AssetFileContextMenu.cpp
--- a/Engine/Plugins/Editor/ContentBrowser/ContentBrowserAssetDataSource/Source/ContentBrowserAssetDataSource/Private/AssetFileContextMenu.cpp
+++ b/Engine/Plugins/Editor/ContentBrowser/ContentBrowserAssetDataSource/Source/ContentBrowserAssetDataSource/Private/AssetFileContextMenu.cpp
@@ -1867,13 +1867,14 @@ void FAssetFileContextMenu::ExecuteBulkExport()
 bool FAssetFileContextMenu::CanExecuteCreateBlueprintUsing() const
 {
 	// Only work if you have a single asset selected
-	if(SelectedAssets.Num() == 1)
+	if (SelectedAssets.Num() == 1)
 	{
-		if (UObject* Asset = SelectedAssets[0].FastGetAsset(true))
+		// There's no need to resolve the class, we don't need to probably worry about blueprint components here for
+		// the component asset brokerage, so just loaded native classes will work.
+		if (UClass* AssetClass = SelectedAssets[0].GetClass(EResolveClass::No))
 		{
-			// See if we know how to make a component from this asset
-			TArray< TSubclassOf<UActorComponent> > ComponentClassList = FComponentAssetBrokerage::GetComponentsForAsset(Asset);
-			return (ComponentClassList.Num() > 0);
+			// Just see if there's a primary component for this asset class.
+			return FComponentAssetBrokerage::GetPrimaryComponentForAsset(AssetClass) != nullptr;
 		}
 	}
 
diff --git a/Engine/Plugins/Editor/DataValidation/Source/DataValidation/Private/DataValidationModule.cpp b/Engine/Plugins/Editor/DataValidation/Source/DataValidation/Private/DataValidationModule.cpp
--- a/Engine/Plugins/Editor/DataValidation/Source/DataValidation/Private/DataValidationModule.cpp
+++ b/Engine/Plugins/Editor/DataValidation/Source/DataValidation/Private/DataValidationModule.cpp
@@ -106,7 +106,7 @@ void FDataValidationModule::RegisterMenus()
 			FSlateIcon(),
 			FToolMenuExecuteAction::CreateLambda([this](const FToolMenuContext& InContext)
 			{
-				if (UContentBrowserAssetContextMenuContext* Context = InContext.FindContext<UContentBrowserAssetContextMenuContext>())
+				if (const UContentBrowserAssetContextMenuContext* Context = InContext.FindContext<UContentBrowserAssetContextMenuContext>())
 				{
 					ValidateAssets(Context->SelectedAssets, false, EDataValidationUsecase::Manual);
 				}
@@ -120,7 +120,7 @@ void FDataValidationModule::RegisterMenus()
 			FSlateIcon(),
 			FToolMenuExecuteAction::CreateLambda([this](const FToolMenuContext& InContext)
 			{
-				if (UContentBrowserAssetContextMenuContext* Context = InContext.FindContext<UContentBrowserAssetContextMenuContext>())
+				if (const UContentBrowserAssetContextMenuContext* Context = InContext.FindContext<UContentBrowserAssetContextMenuContext>())
 				{
 					ValidateAssets(Context->SelectedAssets, true, EDataValidationUsecase::Manual);
 				}
@@ -138,7 +138,7 @@ void FDataValidationModule::RegisterMenus()
 			FSlateIcon(FAppStyle::GetAppStyleSetName(), "Icons.Validate"),
 			FToolMenuExecuteAction::CreateLambda([this](const FToolMenuContext& InContext)
 			{
-				if (UContentBrowserFolderContext* Context = InContext.FindContext<UContentBrowserFolderContext>())
+				if (const UContentBrowserFolderContext* Context = InContext.FindContext<UContentBrowserFolderContext>())
 				{
 					const TArray<FString>& SelectedPaths = Context->GetSelectedPackagePaths();
 
diff --git a/Engine/Plugins/Editor/EngineAssetDefinitions/Source/Private/AssetDefinition_Blueprint.cpp b/Engine/Plugins/Editor/EngineAssetDefinitions/Source/Private/AssetDefinition_Blueprint.cpp
--- a/Engine/Plugins/Editor/EngineAssetDefinitions/Source/Private/AssetDefinition_Blueprint.cpp
+++ b/Engine/Plugins/Editor/EngineAssetDefinitions/Source/Private/AssetDefinition_Blueprint.cpp
@@ -111,7 +111,7 @@ namespace MenuExtension_Blueprint
 			UToolMenu* Menu = UE::ContentBrowser::ExtendToolMenu_AssetContextMenu(UBlueprint::StaticClass());
 	        
 			FToolMenuSection& Section = Menu->FindOrAddSection("GetAssetActions");
-				Section.AddDynamicEntry("GetAssetActions_Blueprint", FNewToolMenuSectionDelegate::CreateLambda([](FToolMenuSection& InSection)
+				Section.AddDynamicEntry(NAME_None, FNewToolMenuSectionDelegate::CreateLambda([](FToolMenuSection& InSection)
 				{
 					if (const UContentBrowserAssetContextMenuContext* Context = UContentBrowserAssetContextMenuContext::FindContextWithAssets(InSection))
 					{
diff --git a/Engine/Plugins/Editor/EngineAssetDefinitions/Source/Private/AssetDefinition_Curve.cpp b/Engine/Plugins/Editor/EngineAssetDefinitions/Source/Private/AssetDefinition_Curve.cpp
--- a/Engine/Plugins/Editor/EngineAssetDefinitions/Source/Private/AssetDefinition_Curve.cpp
+++ b/Engine/Plugins/Editor/EngineAssetDefinitions/Source/Private/AssetDefinition_Curve.cpp
@@ -162,7 +162,7 @@ namespace MenuExtension_CurveBase
 			UToolMenu* Menu = UE::ContentBrowser::ExtendToolMenu_AssetContextMenu(UCurveBase::StaticClass());
 		
 			FToolMenuSection& Section = Menu->FindOrAddSection("GetAssetActions");
-				Section.AddDynamicEntry("GetAssetActions_UCurveBase", FNewToolMenuSectionDelegate::CreateLambda([](FToolMenuSection& InSection)
+				Section.AddDynamicEntry(NAME_None, FNewToolMenuSectionDelegate::CreateLambda([](FToolMenuSection& InSection)
 				{
 					if (const UContentBrowserAssetContextMenuContext* Context = UContentBrowserAssetContextMenuContext::FindContextWithAssets(InSection))
 					{
diff --git a/Engine/Plugins/Editor/EngineAssetDefinitions/Source/Private/AssetDefinition_DataAsset.cpp b/Engine/Plugins/Editor/EngineAssetDefinitions/Source/Private/AssetDefinition_DataAsset.cpp
--- a/Engine/Plugins/Editor/EngineAssetDefinitions/Source/Private/AssetDefinition_DataAsset.cpp
+++ b/Engine/Plugins/Editor/EngineAssetDefinitions/Source/Private/AssetDefinition_DataAsset.cpp
@@ -97,7 +97,7 @@ namespace MenuExtension_DataAsset
 			UToolMenu* Menu = UE::ContentBrowser::ExtendToolMenu_AssetContextMenu(UDataAsset::StaticClass());
 	        
 			FToolMenuSection& Section = Menu->FindOrAddSection("GetAssetActions");
-			Section.AddDynamicEntry("GetAssetActions_DataAsset", FNewToolMenuSectionDelegate::CreateStatic([](FToolMenuSection& InSection)
+			Section.AddDynamicEntry(NAME_None, FNewToolMenuSectionDelegate::CreateStatic([](FToolMenuSection& InSection)
 			{
 				{
 					const TAttribute<FText> Label = LOCTEXT("DataAsset_ChangeClass", "Convert to Different DataAsset Type");
@@ -146,25 +146,4 @@ EAssetCommandResult UAssetDefinition_DataAsset::PerformAssetDiff(const FAssetDif
 	return EAssetCommandResult::Handled;
 }
 
-EAssetCommandResult UAssetDefinition_DataAsset::GetFilters(TArray<FAssetFilterData>& OutFilters) const
-{
-	// AssetDefinitions don't allow filters for abstract classes by default, so we override that behavior to allow
-	// filters for Data Assets for now
-	const TSoftClassPtr<UObject> AssetClassPtr = GetAssetClass();
-	
-	if (const UClass* AssetClass = AssetClassPtr.Get())
-	{
-		FAssetFilterData DefaultFilter;
-		DefaultFilter.Name = AssetClassPtr.ToSoftObjectPath().ToString();
-		DefaultFilter.DisplayText = GetAssetDisplayName();
-		DefaultFilter.Filter.ClassPaths.Add(AssetClassPtr.ToSoftObjectPath().GetAssetPath());
-		DefaultFilter.Filter.bRecursiveClasses = true;
-		OutFilters.Add(MoveTemp(DefaultFilter));
-
-		return EAssetCommandResult::Handled;
-	}
-	
-	return EAssetCommandResult::Unhandled;
-}
-
 #undef LOCTEXT_NAMESPACE
diff --git a/Engine/Plugins/Editor/EngineAssetDefinitions/Source/Private/AssetDefinition_Font.cpp b/Engine/Plugins/Editor/EngineAssetDefinitions/Source/Private/AssetDefinition_Font.cpp
--- a/Engine/Plugins/Editor/EngineAssetDefinitions/Source/Private/AssetDefinition_Font.cpp
+++ b/Engine/Plugins/Editor/EngineAssetDefinitions/Source/Private/AssetDefinition_Font.cpp
@@ -50,7 +50,7 @@ namespace MenuExtension_Font
 			UToolMenu* Menu = UE::ContentBrowser::ExtendToolMenu_AssetContextMenu(UFont::StaticClass());
 		
 			FToolMenuSection& Section = Menu->FindOrAddSection("GetAssetActions");
-			Section.AddDynamicEntry("GetAssetActions_UStaticMesh", FNewToolMenuSectionDelegate::CreateLambda([](FToolMenuSection& InSection)
+			Section.AddDynamicEntry(NAME_None, FNewToolMenuSectionDelegate::CreateLambda([](FToolMenuSection& InSection)
 			{
 				if (const UContentBrowserAssetContextMenuContext* Context = UContentBrowserAssetContextMenuContext::FindContextWithAssets(InSection))
 				{
diff --git a/Engine/Plugins/Editor/EngineAssetDefinitions/Source/Private/AssetDefinition_FontFace.cpp b/Engine/Plugins/Editor/EngineAssetDefinitions/Source/Private/AssetDefinition_FontFace.cpp
--- a/Engine/Plugins/Editor/EngineAssetDefinitions/Source/Private/AssetDefinition_FontFace.cpp
+++ b/Engine/Plugins/Editor/EngineAssetDefinitions/Source/Private/AssetDefinition_FontFace.cpp
@@ -54,7 +54,7 @@ namespace MenuExtension_FontFace
 			UToolMenu* Menu = UE::ContentBrowser::ExtendToolMenu_AssetContextMenu(UFontFace::StaticClass());
 		
 			FToolMenuSection& Section = Menu->FindOrAddSection("GetAssetActions");
-			Section.AddDynamicEntry("GetAssetActions_UStaticMesh", FNewToolMenuSectionDelegate::CreateLambda([](FToolMenuSection& InSection)
+			Section.AddDynamicEntry(NAME_None, FNewToolMenuSectionDelegate::CreateLambda([](FToolMenuSection& InSection)
 			{
 				if (const UContentBrowserAssetContextMenuContext* Context = UContentBrowserAssetContextMenuContext::FindContextWithAssets(InSection))
 				{
diff --git a/Engine/Plugins/Editor/EngineAssetDefinitions/Source/Private/AssetDefinition_Redirector.cpp b/Engine/Plugins/Editor/EngineAssetDefinitions/Source/Private/AssetDefinition_Redirector.cpp
--- a/Engine/Plugins/Editor/EngineAssetDefinitions/Source/Private/AssetDefinition_Redirector.cpp
+++ b/Engine/Plugins/Editor/EngineAssetDefinitions/Source/Private/AssetDefinition_Redirector.cpp
@@ -72,7 +72,7 @@ namespace MenuExtension_Redirector
 			UToolMenu* Menu = UE::ContentBrowser::ExtendToolMenu_AssetContextMenu(UObjectRedirector::StaticClass());
 	        
 			FToolMenuSection& Section = Menu->FindOrAddSection("GetAssetActions");
-			Section.AddDynamicEntry("GetAssetActions_ObjectRedirector", FNewToolMenuSectionDelegate::CreateLambda([](FToolMenuSection& InSection)
+			Section.AddDynamicEntry(NAME_None, FNewToolMenuSectionDelegate::CreateLambda([](FToolMenuSection& InSection)
 			{
 				{
 					const TAttribute<FText> Label = LOCTEXT("Redirector_FindTarget", "Find Target");
diff --git a/Engine/Plugins/Editor/EngineAssetDefinitions/Source/Private/AssetDefinition_StaticMesh.cpp b/Engine/Plugins/Editor/EngineAssetDefinitions/Source/Private/AssetDefinition_StaticMesh.cpp
--- a/Engine/Plugins/Editor/EngineAssetDefinitions/Source/Private/AssetDefinition_StaticMesh.cpp
+++ b/Engine/Plugins/Editor/EngineAssetDefinitions/Source/Private/AssetDefinition_StaticMesh.cpp
@@ -513,7 +513,7 @@ namespace MenuExtension_StaticMesh
 			UToolMenu* Menu = UE::ContentBrowser::ExtendToolMenu_AssetContextMenu(UStaticMesh::StaticClass());
 		
 			FToolMenuSection& Section = Menu->FindOrAddSection("GetAssetActions");
-			Section.AddDynamicEntry("GetAssetActions_UStaticMesh", FNewToolMenuSectionDelegate::CreateLambda([](FToolMenuSection& InSection)
+			Section.AddDynamicEntry(NAME_None, FNewToolMenuSectionDelegate::CreateLambda([](FToolMenuSection& InSection)
 			{
 				if (const UContentBrowserAssetContextMenuContext* Context = UContentBrowserAssetContextMenuContext::FindContextWithAssets(InSection))
 				{
diff --git a/Engine/Plugins/Editor/EngineAssetDefinitions/Source/Private/AssetDefinition_Texture.cpp b/Engine/Plugins/Editor/EngineAssetDefinitions/Source/Private/AssetDefinition_Texture.cpp
--- a/Engine/Plugins/Editor/EngineAssetDefinitions/Source/Private/AssetDefinition_Texture.cpp
+++ b/Engine/Plugins/Editor/EngineAssetDefinitions/Source/Private/AssetDefinition_Texture.cpp
@@ -196,7 +196,7 @@ namespace MenuExtension_Texture
 			UToolMenu* Menu = UE::ContentBrowser::ExtendToolMenu_AssetContextMenu(UTexture::StaticClass());
 		
 			FToolMenuSection& Section = Menu->FindOrAddSection("GetAssetActions");
-				Section.AddDynamicEntry("GetAssetActions_UTexture", FNewToolMenuSectionDelegate::CreateLambda([](FToolMenuSection& InSection)
+				Section.AddDynamicEntry(NAME_None, FNewToolMenuSectionDelegate::CreateLambda([](FToolMenuSection& InSection)
 				{
 					if (const UContentBrowserAssetContextMenuContext* Context = UContentBrowserAssetContextMenuContext::FindContextWithAssets(InSection))
 					{
diff --git a/Engine/Plugins/Editor/EngineAssetDefinitions/Source/Private/AssetDefinition_Texture2D.cpp b/Engine/Plugins/Editor/EngineAssetDefinitions/Source/Private/AssetDefinition_Texture2D.cpp
--- a/Engine/Plugins/Editor/EngineAssetDefinitions/Source/Private/AssetDefinition_Texture2D.cpp
+++ b/Engine/Plugins/Editor/EngineAssetDefinitions/Source/Private/AssetDefinition_Texture2D.cpp
@@ -116,7 +116,7 @@ namespace MenuExtension_Texture2D
 			UToolMenu* Menu = UE::ContentBrowser::ExtendToolMenu_AssetContextMenu(UTexture2D::StaticClass());
 	        
 			FToolMenuSection& Section = Menu->FindOrAddSection("GetAssetActions");
-			Section.AddDynamicEntry("GetAssetActions_UTexture2D", FNewToolMenuSectionDelegate::CreateLambda([](FToolMenuSection& InSection)
+			Section.AddDynamicEntry(NAME_None, FNewToolMenuSectionDelegate::CreateLambda([](FToolMenuSection& InSection)
 			{
 				const UContentBrowserAssetContextMenuContext* Context = UContentBrowserAssetContextMenuContext::FindContextWithAssets(InSection);
 				
diff --git a/Engine/Plugins/Editor/EngineAssetDefinitions/Source/Private/AssetDefinition_TextureCube.cpp b/Engine/Plugins/Editor/EngineAssetDefinitions/Source/Private/AssetDefinition_TextureCube.cpp
--- a/Engine/Plugins/Editor/EngineAssetDefinitions/Source/Private/AssetDefinition_TextureCube.cpp
+++ b/Engine/Plugins/Editor/EngineAssetDefinitions/Source/Private/AssetDefinition_TextureCube.cpp
@@ -55,7 +55,7 @@ namespace MenuExtension_TextureCubeArray
 			UToolMenu* Menu = UE::ContentBrowser::ExtendToolMenu_AssetContextMenu(UTextureCube::StaticClass());
 	        
 			FToolMenuSection& Section = Menu->FindOrAddSection("GetAssetActions");
-			Section.AddDynamicEntry("GetAssetActions_UTextureCube", FNewToolMenuSectionDelegate::CreateLambda([](FToolMenuSection& InSection)
+			Section.AddDynamicEntry(NAME_None, FNewToolMenuSectionDelegate::CreateLambda([](FToolMenuSection& InSection)
 			{
 				{
 					const TAttribute<FText> Label = LOCTEXT("Texture_TextureCubeArray", "Create Texture Cube Array");
diff --git a/Engine/Plugins/Editor/EngineAssetDefinitions/Source/Private/AssetDefinition_TextureRenderTarget.cpp b/Engine/Plugins/Editor/EngineAssetDefinitions/Source/Private/AssetDefinition_TextureRenderTarget.cpp
--- a/Engine/Plugins/Editor/EngineAssetDefinitions/Source/Private/AssetDefinition_TextureRenderTarget.cpp
+++ b/Engine/Plugins/Editor/EngineAssetDefinitions/Source/Private/AssetDefinition_TextureRenderTarget.cpp
@@ -76,7 +76,7 @@ namespace MenuExtension_TextureRenderTarget
 			UToolMenu* Menu = UE::ContentBrowser::ExtendToolMenu_AssetContextMenu(UTextureRenderTarget::StaticClass());
 	        
 			FToolMenuSection& Section = Menu->FindOrAddSection("GetAssetActions");
-			Section.AddDynamicEntry("GetAssetActions_UTextureRenderTarget", FNewToolMenuSectionDelegate::CreateLambda([](FToolMenuSection& InSection)
+			Section.AddDynamicEntry(NAME_None, FNewToolMenuSectionDelegate::CreateLambda([](FToolMenuSection& InSection)
 			{
 				{
 					const TAttribute<FText> Label = LOCTEXT("TextureRenderTarget_CreateStatic", "Create Static Texture");
diff --git a/Engine/Plugins/Editor/EngineAssetDefinitions/Source/Private/Material/AssetDefinition_MaterialFunction.cpp b/Engine/Plugins/Editor/EngineAssetDefinitions/Source/Private/Material/AssetDefinition_MaterialFunction.cpp
--- a/Engine/Plugins/Editor/EngineAssetDefinitions/Source/Private/Material/AssetDefinition_MaterialFunction.cpp
+++ b/Engine/Plugins/Editor/EngineAssetDefinitions/Source/Private/Material/AssetDefinition_MaterialFunction.cpp
@@ -99,7 +99,7 @@ namespace MenuExtension_MaterialFunction
 			UToolMenu* Menu = UE::ContentBrowser::ExtendToolMenu_AssetContextMenu(UMaterialFunction::StaticClass());
 	        
 			FToolMenuSection& Section = Menu->FindOrAddSection("GetAssetActions");
-			Section.AddDynamicEntry("GetAssetActions_NewMaterialFunctionInstance", FNewToolMenuSectionDelegate::CreateLambda([](FToolMenuSection& InSection)
+			Section.AddDynamicEntry(NAME_None, FNewToolMenuSectionDelegate::CreateLambda([](FToolMenuSection& InSection)
 			{
 				{
 					const TAttribute<FText> Label = NSLOCTEXT("AssetTypeActions", "Material_NewMFI", "Create Function Instance");
@@ -111,7 +111,7 @@ namespace MenuExtension_MaterialFunction
 				}
 			}));
 			
-			Section.AddDynamicEntry("GetAssetActions_UMaterialFunction", FNewToolMenuSectionDelegate::CreateLambda([](FToolMenuSection& InSection)
+			Section.AddDynamicEntry(NAME_None, FNewToolMenuSectionDelegate::CreateLambda([](FToolMenuSection& InSection)
 			{
 				if (FEditorDelegates::OnOpenReferenceViewer.IsBound())
 				{
@@ -152,7 +152,7 @@ namespace MenuExtension_MaterialFunctionMaterialLayer
 			for (UToolMenu* Menu : Menus)
 			{
 				FToolMenuSection& Section = Menu->FindOrAddSection("GetAssetActions");
-				Section.AddDynamicEntry("GetAssetActions_NewMaterialFunctionInstance", FNewToolMenuSectionDelegate::CreateLambda([](FToolMenuSection& InSection)
+				Section.AddDynamicEntry(NAME_None, FNewToolMenuSectionDelegate::CreateLambda([](FToolMenuSection& InSection)
 				{
 					{
 						const TAttribute<FText> Label = NSLOCTEXT("AssetTypeActions", "Material_NewMLI", "Create Layer Instance");
@@ -193,7 +193,7 @@ namespace MenuExtension_MaterialFunctionMaterialLayerBlend
 			for (UToolMenu* Menu : Menus)
 			{
 				FToolMenuSection& Section = Menu->FindOrAddSection("GetAssetActions");
-				Section.AddDynamicEntry("GetAssetActions_NewMaterialFunctionInstance", FNewToolMenuSectionDelegate::CreateLambda([](FToolMenuSection& InSection)
+				Section.AddDynamicEntry(NAME_None, FNewToolMenuSectionDelegate::CreateLambda([](FToolMenuSection& InSection)
 				{
 					{
 						const TAttribute<FText> Label = NSLOCTEXT("AssetTypeActions", "Material_NewMBI", "Create Blend Instance");
diff --git a/Engine/Plugins/Editor/EngineAssetDefinitions/Source/Private/Material/AssetDefinition_MaterialFunctionInstance.cpp b/Engine/Plugins/Editor/EngineAssetDefinitions/Source/Private/Material/AssetDefinition_MaterialFunctionInstance.cpp
--- a/Engine/Plugins/Editor/EngineAssetDefinitions/Source/Private/Material/AssetDefinition_MaterialFunctionInstance.cpp
+++ b/Engine/Plugins/Editor/EngineAssetDefinitions/Source/Private/Material/AssetDefinition_MaterialFunctionInstance.cpp
@@ -69,7 +69,7 @@ namespace MenuExtension_MaterialFunctionInstance
 			UToolMenu* Menu = UE::ContentBrowser::ExtendToolMenu_AssetContextMenu(UMaterialFunctionInstance::StaticClass());
 	        
 			FToolMenuSection& Section = Menu->FindOrAddSection("GetAssetActions");
-			Section.AddDynamicEntry("GetAssetActions_UMaterialFunctionInstance", FNewToolMenuSectionDelegate::CreateLambda([](FToolMenuSection& InSection)
+			Section.AddDynamicEntry(NAME_None, FNewToolMenuSectionDelegate::CreateLambda([](FToolMenuSection& InSection)
 			{
 				{
 					const TAttribute<FText> Label = LOCTEXT("MaterialFunctionInstance_FindParent", "Find Parent");
diff --git a/Engine/Plugins/Editor/EngineAssetDefinitions/Source/Private/Material/AssetDefinition_MaterialInstanceConstant.cpp b/Engine/Plugins/Editor/EngineAssetDefinitions/Source/Private/Material/AssetDefinition_MaterialInstanceConstant.cpp
--- a/Engine/Plugins/Editor/EngineAssetDefinitions/Source/Private/Material/AssetDefinition_MaterialInstanceConstant.cpp
+++ b/Engine/Plugins/Editor/EngineAssetDefinitions/Source/Private/Material/AssetDefinition_MaterialInstanceConstant.cpp
@@ -55,7 +55,7 @@ namespace
 			UToolMenu* Menu = UE::ContentBrowser::ExtendToolMenu_AssetContextMenu(UMaterialInstanceConstant::StaticClass());
 	        
 			FToolMenuSection& Section = Menu->FindOrAddSection("GetAssetActions");
-			Section.AddDynamicEntry("GetAssetActions_MaterialInstanceConstant", FNewToolMenuSectionDelegate::CreateLambda([](FToolMenuSection& InSection)
+			Section.AddDynamicEntry(NAME_None, FNewToolMenuSectionDelegate::CreateLambda([](FToolMenuSection& InSection)
 			{
 				{
 					const TAttribute<FText> Label = LOCTEXT("MaterialInstanceConstant_FindParent", "Find Parent");
diff --git a/Engine/Plugins/Editor/EngineAssetDefinitions/Source/Private/Material/AssetDefinition_MaterialInterface.cpp b/Engine/Plugins/Editor/EngineAssetDefinitions/Source/Private/Material/AssetDefinition_MaterialInterface.cpp
--- a/Engine/Plugins/Editor/EngineAssetDefinitions/Source/Private/Material/AssetDefinition_MaterialInterface.cpp
+++ b/Engine/Plugins/Editor/EngineAssetDefinitions/Source/Private/Material/AssetDefinition_MaterialInterface.cpp
@@ -61,7 +61,7 @@ namespace MenuExtension_MaterialInterface
 			UToolMenu* Menu = UE::ContentBrowser::ExtendToolMenu_AssetContextMenu(UMaterialInterface::StaticClass());
 	        
 			FToolMenuSection& Section = Menu->FindOrAddSection("GetAssetActions");
-			Section.AddDynamicEntry("GetAssetActions_UMaterialInterface", FNewToolMenuSectionDelegate::CreateLambda([](FToolMenuSection& InSection)
+			Section.AddDynamicEntry(NAME_None, FNewToolMenuSectionDelegate::CreateLambda([](FToolMenuSection& InSection)
 			{
 				{
 					const TAttribute<FText> Label = LOCTEXT("Material_NewMIC", "Create Material Instance");
diff --git a/Engine/Plugins/Editor/EngineAssetDefinitions/Source/Public/AssetDefinition_DataAsset.h b/Engine/Plugins/Editor/EngineAssetDefinitions/Source/Public/AssetDefinition_DataAsset.h
--- a/Engine/Plugins/Editor/EngineAssetDefinitions/Source/Public/AssetDefinition_DataAsset.h
+++ b/Engine/Plugins/Editor/EngineAssetDefinitions/Source/Public/AssetDefinition_DataAsset.h
@@ -15,11 +15,14 @@ class UAssetDefinition_DataAsset : public UAssetDefinitionDefault
 
 public:
 	// UAssetDefinition Implementation
+	UAssetDefinition_DataAsset()
+	{
+		IncludeClassInFilter = EIncludeClassInFilter::Always;
+	}
 	virtual FText GetAssetDisplayName() const override { return NSLOCTEXT("AssetDefinition", "DataAsset", "Data Asset"); }
 	virtual FText GetAssetDisplayName(const FAssetData& AssetData) const override;
 	virtual FLinearColor GetAssetColor() const override { return FLinearColor(FColor(201, 29, 85)); }
 	virtual TSoftClassPtr<UObject> GetAssetClass() const override { return UDataAsset::StaticClass(); }
 	
 	virtual EAssetCommandResult PerformAssetDiff(const FAssetDiffArgs& DiffArgs) const override;
-	virtual EAssetCommandResult GetFilters(TArray<FAssetFilterData>& OutFilters) const override;
 };
diff --git a/Engine/Plugins/Experimental/MotoSynth/Source/MotoSynthEditor/Private/SoundWaveAssetActionExtenderMotoSynth.cpp b/Engine/Plugins/Experimental/MotoSynth/Source/MotoSynthEditor/Private/SoundWaveAssetActionExtenderMotoSynth.cpp
--- a/Engine/Plugins/Experimental/MotoSynth/Source/MotoSynthEditor/Private/SoundWaveAssetActionExtenderMotoSynth.cpp
+++ b/Engine/Plugins/Experimental/MotoSynth/Source/MotoSynthEditor/Private/SoundWaveAssetActionExtenderMotoSynth.cpp
@@ -27,7 +27,7 @@ void FMotoSynthExtension::RegisterMenus()
 	UToolMenu* Menu = UToolMenus::Get()->ExtendMenu("ContentBrowser.AssetContextMenu.SoundWave");
 	FToolMenuSection& Section = Menu->FindOrAddSection("GetAssetActions");
 
-	Section.AddDynamicEntry("SoundWaveAssetConversion", FNewToolMenuSectionDelegate::CreateLambda([](FToolMenuSection& InSection)
+	Section.AddDynamicEntry(NAME_None, FNewToolMenuSectionDelegate::CreateLambda([](FToolMenuSection& InSection)
 	{
 		UContentBrowserAssetContextMenuContext* Context = InSection.FindContext<UContentBrowserAssetContextMenuContext>();
 		if (!Context || Context->SelectedAssets.IsEmpty())
diff --git a/Engine/Plugins/Runtime/SoundCueTemplates/Source/SoundCueTemplatesEditor/Private/AssetDefinition_SoundCueTemplate.cpp b/Engine/Plugins/Runtime/SoundCueTemplates/Source/SoundCueTemplatesEditor/Private/AssetDefinition_SoundCueTemplate.cpp
new file mode 100644
--- /dev/null
+++ b/Engine/Plugins/Runtime/SoundCueTemplates/Source/SoundCueTemplatesEditor/Private/AssetDefinition_SoundCueTemplate.cpp
@@ -0,0 +1,132 @@
+// Copyright Epic Games, Inc. All Rights Reserved.
+
+#include "AssetDefinition_SoundCueTemplate.h"
+
+#include "AssetToolsModule.h"
+#include "SoundCueTemplate.h"
+#include "SoundCueTemplateFactory.h"
+
+#include "AudioEditorSettings.h"
+#include "Components/AudioComponent.h"
+#include "ToolMenus.h"
+#include "ToolMenuSection.h"
+#include "ContentBrowserModule.h"
+#include "ContentBrowserMenuContexts.h"
+#include "IContentBrowserSingleton.h"
+#include "ObjectEditorUtils.h"
+#include "Sound/SoundWaveProcedural.h"
+
+#define LOCTEXT_NAMESPACE "AssetTypeActions"
+
+TConstArrayView<FAssetCategoryPath> UAssetDefinition_SoundCueTemplate::GetAssetCategories() const
+{
+	static const auto Pinned_Categories = { EAssetCategoryPaths::Audio };
+	static const auto Categories = { EAssetCategoryPaths::Audio / LOCTEXT("AssetSoundCueSubMenu", "Source") };
+	
+	if (GetDefault<UAudioEditorSettings>()->bPinSoundCueTemplateInAssetMenu)
+	{
+		return Pinned_Categories;
+	}
+	
+	return Categories;
+}
+
+// Menu Extensions
+//--------------------------------------------------------------------
+
+namespace MenuExtension_SoundCueTemplate
+{
+	void ExecuteCopyToSoundCue(const FToolMenuContext& MenuContext)
+	{
+		const FAssetToolsModule& AssetToolsModule = FModuleManager::Get().LoadModuleChecked<FAssetToolsModule>("AssetTools");
+		
+		if (const UContentBrowserAssetContextMenuContext* Context = UContentBrowserAssetContextMenuContext::FindContextWithAssets(MenuContext))
+		{
+			for (USoundCueTemplate* Object : Context->LoadSelectedObjects<USoundCueTemplate>())
+			{
+				FString Name;
+				FString PackagePath;
+				AssetToolsModule.Get().CreateUniqueAssetName(Object->GetOutermost()->GetName(), TEXT(""), PackagePath, Name);
+
+				if (USoundCueTemplateCopyFactory* Factory = NewObject<USoundCueTemplateCopyFactory>())
+				{
+					Factory->SoundCueTemplate = TWeakObjectPtr<USoundCueTemplate>(Object);
+					FContentBrowserModule& ContentBrowserModule = FModuleManager::LoadModuleChecked<FContentBrowserModule>("ContentBrowser");
+					ContentBrowserModule.Get().CreateNewAsset(Name, FPackageName::GetLongPackagePath(PackagePath), USoundCue::StaticClass(), Factory);
+				}
+			}
+		}
+	}
+
+	void ExecuteCreateSoundCueTemplate(const struct FToolMenuContext& MenuContext)
+	{
+		const FAssetToolsModule& AssetToolsModule = FModuleManager::Get().LoadModuleChecked<FAssetToolsModule>("AssetTools");
+		
+		if (const UContentBrowserAssetContextMenuContext* Context = UContentBrowserAssetContextMenuContext::FindContextWithAssets(MenuContext))
+		{
+			TArray<USoundWave*> SoundWaves = Context->LoadSelectedObjectsIf<USoundWave>([](const FAssetData& AssetData) { return !AssetData.IsInstanceOf<USoundWaveProcedural>(); });
+			if (!SoundWaves.IsEmpty())
+			{
+				FString PackagePath;
+				FString Name;
+				
+				AssetToolsModule.Get().CreateUniqueAssetName(SoundWaves[0]->GetOutermost()->GetName(), TEXT(""), PackagePath, Name);
+
+				USoundCueTemplateFactory* Factory = NewObject<USoundCueTemplateFactory>();
+				Factory->SoundWaves = TArray<TWeakObjectPtr<USoundWave>>(SoundWaves);
+				Factory->ConfigureProperties();
+				Name = Factory->GetDefaultNewAssetName();
+
+				const FContentBrowserModule& ContentBrowserModule = FModuleManager::Get().LoadModuleChecked<FContentBrowserModule>("ContentBrowser");
+				ContentBrowserModule.Get().CreateNewAsset(Name, FPackageName::GetLongPackagePath(PackagePath), USoundCueTemplate::StaticClass(), Factory);
+			}
+		}
+	}
+
+	static FDelayedAutoRegisterHelper DelayedAutoRegister(EDelayedRegisterRunPhase::EndOfEngineInit, []{ 
+		UToolMenus::RegisterStartupCallback(FSimpleMulticastDelegate::FDelegate::CreateLambda([]()
+		{
+			FToolMenuOwnerScoped OwnerScoped(UE_MODULE_NAME);
+
+			{
+				UToolMenu* Menu = UE::ContentBrowser::ExtendToolMenu_AssetContextMenu(USoundWave::StaticClass());
+		
+				FToolMenuSection& Section = Menu->FindOrAddSection("GetAssetActions");
+				Section.AddDynamicEntry(NAME_None, FNewToolMenuSectionDelegate::CreateLambda([](FToolMenuSection& InSection)
+				{
+					const TAttribute<FText> Label = LOCTEXT("SoundWave_CreateSoundCueTemplate", "Create SoundCueTemplate");
+					const TAttribute<FText> ToolTip = LOCTEXT("SoundWave_CreateSoundCueTemplateToolTip", "Creates a SoundCueTemplate from the selected sound waves.");
+					const FSlateIcon Icon = FSlateIcon(FAppStyle::GetAppStyleSetName(), "ClassIcon.SoundCue");
+					
+					FToolUIAction UIAction;
+					UIAction.ExecuteAction = FToolMenuExecuteAction::CreateStatic(&ExecuteCreateSoundCueTemplate);
+					InSection.AddMenuEntry("SoundWave_CreateSoundCueTemplate", Label, ToolTip, Icon, UIAction);
+				}));
+			}
+
+
+			{
+				UToolMenu* Menu = UE::ContentBrowser::ExtendToolMenu_AssetContextMenu(USoundCueTemplate::StaticClass());
+				
+				FToolMenuSection& Section = Menu->FindOrAddSection("GetAssetActions");
+				Section.AddDynamicEntry(NAME_None, FNewToolMenuSectionDelegate::CreateLambda([](FToolMenuSection& InSection)
+				{
+					if (const UContentBrowserAssetContextMenuContext* Context = UContentBrowserAssetContextMenuContext::FindContextWithAssets(InSection))
+					{
+						{
+							const TAttribute<FText> Label = LOCTEXT("SoundCueTemplate_CopyToSoundCue", "Copy To Sound Cue");
+							const TAttribute<FText> ToolTip = LOCTEXT("SoundCueTemplate_CopyToSoundCueTooltip", "Exports a Sound Cue Template to a Sound Cue.");
+							const FSlateIcon Icon = FSlateIcon(FAppStyle::GetAppStyleSetName(), "ClassIcon.SoundCue");
+							
+							FToolUIAction UIAction;
+                            UIAction.ExecuteAction = FToolMenuExecuteAction::CreateStatic(&ExecuteCopyToSoundCue);
+							InSection.AddMenuEntry("SoundCueTemplate_CopyToSoundCue", Label, ToolTip, Icon, UIAction);
+						}
+					}
+				}));
+			}
+		}));
+	});
+}
+
+#undef LOCTEXT_NAMESPACE
diff --git a/Engine/Plugins/Runtime/SoundCueTemplates/Source/SoundCueTemplatesEditor/Private/AssetDefinition_SoundCueTemplate.h b/Engine/Plugins/Runtime/SoundCueTemplates/Source/SoundCueTemplatesEditor/Private/AssetDefinition_SoundCueTemplate.h
new file mode 100644
--- /dev/null
+++ b/Engine/Plugins/Runtime/SoundCueTemplates/Source/SoundCueTemplatesEditor/Private/AssetDefinition_SoundCueTemplate.h
@@ -0,0 +1,36 @@
+// Copyright Epic Games, Inc. All Rights Reserved.
+
+#pragma once
+
+#include "CoreMinimal.h"
+#include "SoundCueTemplate.h"
+#include "AssetTypeActions/AssetDefinition_SoundBase.h"
+
+#include "AssetDefinition_SoundCueTemplate.generated.h"
+
+class USoundCueTemplate;
+struct FToolMenuContext;
+
+UCLASS()
+class UAssetDefinition_SoundCueTemplate : public UAssetDefinition_SoundBase
+{
+	GENERATED_BODY()
+
+public:
+	// UAssetDefinition Begin
+	virtual FText GetAssetDisplayName() const override { return NSLOCTEXT("AssetTypeActions", "AssetTypeActions_SoundCueTemplate", "Sound Cue Template"); }
+	virtual FLinearColor GetAssetColor() const override { return FLinearColor(FColor(255, 0, 0)); }
+	virtual TSoftClassPtr<UObject> GetAssetClass() const override { return USoundCueTemplate::StaticClass(); }
+	virtual TConstArrayView<FAssetCategoryPath> GetAssetCategories() const override;
+	// UAssetDefinition End
+};
+
+class FAssetActionExtender_SoundCueTemplate
+{
+public:
+	static void RegisterMenus();
+private:
+	static void ExecuteCreateSoundCueTemplate(const FToolMenuContext& MenuContext);
+	/** Converts the provided SoundCue Template to a fully-modifiable SoundCue */
+	static void ExecuteCopyToSoundCue(const FToolMenuContext& MenuContext);
+};
\ No newline at end of file
diff --git a/Engine/Plugins/Runtime/SoundCueTemplates/Source/SoundCueTemplatesEditor/Private/AssetTypeActions_SoundCueTemplate.cpp b/Engine/Plugins/Runtime/SoundCueTemplates/Source/SoundCueTemplatesEditor/Private/AssetTypeActions_SoundCueTemplate.cpp
deleted file mode 100644
diff --git a/Engine/Plugins/Runtime/SoundCueTemplates/Source/SoundCueTemplatesEditor/Private/AssetTypeActions_SoundCueTemplate.h b/Engine/Plugins/Runtime/SoundCueTemplates/Source/SoundCueTemplatesEditor/Private/AssetTypeActions_SoundCueTemplate.h
deleted file mode 100644
diff --git a/Engine/Plugins/Runtime/SoundCueTemplates/Source/SoundCueTemplatesEditor/Private/SoundCueTemplatesEditorModule.cpp b/Engine/Plugins/Runtime/SoundCueTemplates/Source/SoundCueTemplatesEditor/Private/SoundCueTemplatesEditorModule.cpp
--- a/Engine/Plugins/Runtime/SoundCueTemplates/Source/SoundCueTemplatesEditor/Private/SoundCueTemplatesEditorModule.cpp
+++ b/Engine/Plugins/Runtime/SoundCueTemplates/Source/SoundCueTemplatesEditor/Private/SoundCueTemplatesEditorModule.cpp
@@ -3,38 +3,18 @@
 #include "SoundCueTemplatesEditorModule.h"
 
 #include "Modules/ModuleManager.h"
-#include "AssetToolsModule.h"
 #include "ToolMenus.h"
 #include "HAL/LowLevelMemTracker.h"
 
-#include "AssetTypeActions_SoundCueTemplate.h"
-
 void FSoundCueTemplatesEditorModule::StartupModule()
 {
 	LLM_SCOPE(ELLMTag::AudioMisc);
-	RegisterAssetActions();
-	UToolMenus::RegisterStartupCallback(FSimpleMulticastDelegate::FDelegate::CreateRaw(this, &FSoundCueTemplatesEditorModule::RegisterMenus));
 }
 
 void FSoundCueTemplatesEditorModule::ShutdownModule()
 {
 	LLM_SCOPE(ELLMTag::AudioMisc);
-	UToolMenus::UnRegisterStartupCallback(this);
 	UToolMenus::UnregisterOwner(UE_MODULE_NAME);
 }
 
-void FSoundCueTemplatesEditorModule::RegisterAssetActions()
-{
-	LLM_SCOPE(ELLMTag::AudioMisc);
-	// Register the audio editor asset type actions
-	IAssetTools& AssetTools = FModuleManager::LoadModuleChecked<FAssetToolsModule>("AssetTools").Get();
-
-	AssetTools.RegisterAssetTypeActions(MakeShared<FAssetTypeActions_SoundCueTemplate>());
-}
-
-void FSoundCueTemplatesEditorModule::RegisterMenus()
-{
-	FAssetActionExtender_SoundCueTemplate::RegisterMenus();
-}
-
 IMPLEMENT_MODULE(FSoundCueTemplatesEditorModule, SoundCueTemplatesEditor);
diff --git a/Engine/Plugins/Runtime/SoundCueTemplates/Source/SoundCueTemplatesEditor/Public/SoundCueTemplatesEditorModule.h b/Engine/Plugins/Runtime/SoundCueTemplates/Source/SoundCueTemplatesEditor/Public/SoundCueTemplatesEditorModule.h
--- a/Engine/Plugins/Runtime/SoundCueTemplates/Source/SoundCueTemplatesEditor/Public/SoundCueTemplatesEditorModule.h
+++ b/Engine/Plugins/Runtime/SoundCueTemplates/Source/SoundCueTemplatesEditor/Public/SoundCueTemplatesEditorModule.h
@@ -13,8 +13,4 @@ public:
 	/** IModuleInterface implementation */
 	virtual void StartupModule() override;
 	virtual void ShutdownModule() override;
-
-private:
-	void RegisterAssetActions();
-	void RegisterMenus();
 };
diff --git a/Engine/Plugins/Runtime/SoundCueTemplates/Source/SoundCueTemplatesEditor/SoundCueTemplatesEditor.Build.cs b/Engine/Plugins/Runtime/SoundCueTemplates/Source/SoundCueTemplatesEditor/SoundCueTemplatesEditor.Build.cs
--- a/Engine/Plugins/Runtime/SoundCueTemplates/Source/SoundCueTemplatesEditor/SoundCueTemplatesEditor.Build.cs
+++ b/Engine/Plugins/Runtime/SoundCueTemplates/Source/SoundCueTemplatesEditor/SoundCueTemplatesEditor.Build.cs
@@ -14,6 +14,7 @@ public class SoundCueTemplatesEditor : ModuleRules
 				"Core",
 				"SoundCueTemplates",
 				"AudioEditor",
+				"AssetDefinition",
 				"ToolMenus"
 			}
 		);
diff --git a/Engine/Plugins/Runtime/SoundUtilities/Source/SoundUtilitiesEditor/Private/SoundWaveAssetActionExtender.cpp b/Engine/Plugins/Runtime/SoundUtilities/Source/SoundUtilitiesEditor/Private/SoundWaveAssetActionExtender.cpp
--- a/Engine/Plugins/Runtime/SoundUtilities/Source/SoundUtilitiesEditor/Private/SoundWaveAssetActionExtender.cpp
+++ b/Engine/Plugins/Runtime/SoundUtilities/Source/SoundUtilitiesEditor/Private/SoundWaveAssetActionExtender.cpp
@@ -27,7 +27,7 @@ void FSoundWaveAssetActionExtender::RegisterMenus()
 	UToolMenu* Menu = UToolMenus::Get()->ExtendMenu("ContentBrowser.AssetContextMenu.SoundWave");
 	FToolMenuSection& Section = Menu->FindOrAddSection("GetAssetActions");
 	
-	Section.AddDynamicEntry("SoundWaveAsset", FNewToolMenuSectionDelegate::CreateLambda([](FToolMenuSection& InSection)
+	Section.AddDynamicEntry(NAME_None, FNewToolMenuSectionDelegate::CreateLambda([](FToolMenuSection& InSection)
 	{
 		if (UContentBrowserAssetContextMenuContext* Context = InSection.FindContext<UContentBrowserAssetContextMenuContext>())
 		{
diff --git a/Engine/Source/Developer/ToolMenus/Private/ToolMenus.cpp b/Engine/Source/Developer/ToolMenus/Private/ToolMenus.cpp
--- a/Engine/Source/Developer/ToolMenus/Private/ToolMenus.cpp
+++ b/Engine/Source/Developer/ToolMenus/Private/ToolMenus.cpp
@@ -1896,10 +1896,16 @@ UToolMenu* UToolMenus::GenerateMenuFromHierarchy(const TArray<UToolMenu*>& Hiera
 	return GeneratedMenu;
 }
 
-TSharedRef< class SWidget > UToolMenus::GenerateWidget(const FName InName, const FToolMenuContext& InMenuContext)
+TSharedRef<SWidget> UToolMenus::GenerateWidget(const FName InName, const FToolMenuContext& InMenuContext)
 {
+	OnPreGenerateWidget.Broadcast(InName, InMenuContext);
+
 	UToolMenu* Generated = GenerateMenu(InName, InMenuContext);
-	return GenerateWidget(Generated);
+	TSharedRef<SWidget> Result = GenerateWidget(Generated);
+
+	OnPostGenerateWidget.Broadcast(InName, Generated);
+	
+	return Result;
 }
 
 TSharedRef<SWidget> UToolMenus::GenerateWidget(const TArray<UToolMenu*>& Hierarchy, const FToolMenuContext& InMenuContext)
diff --git a/Engine/Source/Developer/ToolMenus/Public/ToolMenus.h b/Engine/Source/Developer/ToolMenus/Public/ToolMenus.h
--- a/Engine/Source/Developer/ToolMenus/Public/ToolMenus.h
+++ b/Engine/Source/Developer/ToolMenus/Public/ToolMenus.h
@@ -123,9 +123,6 @@ public:
 	 */
 	TSharedRef<SWidget> GenerateWidget(const FName Name, const FToolMenuContext& InMenuContext);
 
-
-
-
 	/**
 	 * Finds an existing menu that has been registered or extended.
 	 * @param	Name	Name of the menu to find.
@@ -283,6 +280,12 @@ public:
 	DECLARE_DELEGATE_OneParam(FEditMenuDelegate, class UToolMenu*);
 	FEditMenuDelegate EditMenuDelegate;
 
+	DECLARE_MULTICAST_DELEGATE_TwoParams(FGenerateWidgetEvent, const FName InName, const FToolMenuContext& InMenuContext);
+	/** Called before we generate a menu widget. */
+	FGenerateWidgetEvent OnPreGenerateWidget;
+	/** Called after we generate a menu widget. */
+	FGenerateWidgetEvent OnPostGenerateWidget;
+
 	/** Icon to display in menus for command to open menu editor */
 	FSlateIcon EditMenuIcon;
 
diff --git a/Engine/Source/Editor/AssetDefinition/Private/AssetDefinition.cpp b/Engine/Source/Editor/AssetDefinition/Private/AssetDefinition.cpp
--- a/Engine/Source/Editor/AssetDefinition/Private/AssetDefinition.cpp
+++ b/Engine/Source/Editor/AssetDefinition/Private/AssetDefinition.cpp
@@ -79,7 +79,7 @@ EAssetCommandResult UAssetDefinition::GetFilters(TArray<FAssetFilterData>& OutFi
 	{
 		// By default we don't advertise filtering if the class is abstract for the asset definition.  Odds are,
 		// if they've registered an abstract class as an asset definition, they mean to use it for subclasses.
-		if (!AssetClass->HasAnyClassFlags(CLASS_Abstract))
+		if (IncludeClassInFilter == EIncludeClassInFilter::Always || (IncludeClassInFilter == EIncludeClassInFilter::IfClassIsNotAbstract && !AssetClass->HasAnyClassFlags(CLASS_Abstract)))
 		{
 			FAssetFilterData DefaultFilter;
 			DefaultFilter.Name = AssetClassPtr.ToSoftObjectPath().ToString();
diff --git a/Engine/Source/Editor/AssetDefinition/Public/AssetDefinition.h b/Engine/Source/Editor/AssetDefinition/Public/AssetDefinition.h
--- a/Engine/Source/Editor/AssetDefinition/Public/AssetDefinition.h
+++ b/Engine/Source/Editor/AssetDefinition/Public/AssetDefinition.h
@@ -225,6 +225,8 @@ struct ASSETDEFINITION_API FAssetCategoryPath
 	FName GetSubCategory() const { return HasSubCategory() ? CategoryPath[1].Key : NAME_None; }
 	FText GetSubCategoryText() const { return HasSubCategory() ? CategoryPath[1].Value : FText::GetEmpty(); }
 	
+	FAssetCategoryPath operator / (const FText& SubCategory) const { return FAssetCategoryPath(*this, SubCategory); }
+	
 private:
 	TArray<TPair<FName, FText>> CategoryPath;
 };
@@ -278,6 +280,12 @@ public:
 
 class UAssetDefinitionRegistry;
 
+enum class EIncludeClassInFilter : uint8
+{
+	IfClassIsNotAbstract,
+	Always
+};
+
 /**
  * Asset Definitions represent top level assets that are known to the editor.
  */
@@ -436,5 +444,8 @@ public:
 protected:
 	virtual bool CanRegisterStatically() const;
 
+protected:
+	EIncludeClassInFilter IncludeClassInFilter = EIncludeClassInFilter::IfClassIsNotAbstract;
+
 	friend class UAssetDefinitionRegistry;
 };
\ No newline at end of file
diff --git a/Engine/Source/Editor/AudioEditor/AudioEditor.Build.cs b/Engine/Source/Editor/AudioEditor/AudioEditor.Build.cs
--- a/Engine/Source/Editor/AudioEditor/AudioEditor.Build.cs
+++ b/Engine/Source/Editor/AudioEditor/AudioEditor.Build.cs
@@ -18,8 +18,7 @@ public class AudioEditor : ModuleRules
 				"DeveloperSettings",
 				"UMGEditor",
 				"AudioExtensions",
-				"AudioLinkEngine",
-				"AssetDefinition"
+				"AudioLinkEngine"
 			}
 		);
 
@@ -30,6 +29,7 @@ public class AudioEditor : ModuleRules
 				"Core",
 				"CoreUObject",
 				"ApplicationCore",
+				"AssetDefinition",
 				"AudioMixer",
 				"SignalProcessing",
 				"InputCore",
@@ -38,7 +38,6 @@ public class AudioEditor : ModuleRules
 				"UnrealEd",
 				"Slate",
 				"SlateCore",
-				
 				"RenderCore",
 				"LevelEditor",
 				"Landscape",
diff --git a/Engine/Source/Editor/AudioEditor/Private/AssetTypeActions/AssetDefinition_SoundBase.cpp b/Engine/Source/Editor/AudioEditor/Private/AssetTypeActions/AssetDefinition_SoundBase.cpp
new file mode 100644
--- /dev/null
+++ b/Engine/Source/Editor/AudioEditor/Private/AssetTypeActions/AssetDefinition_SoundBase.cpp
@@ -0,0 +1,439 @@
+// Copyright Epic Games, Inc. All Rights Reserved.
+
+#include "AssetTypeActions/AssetDefinition_SoundBase.h"
+#include "AudioDeviceManager.h"
+#include "ContentBrowserMenuContexts.h"
+#include "Audio/AudioDebug.h"
+#include "Sound/SoundBase.h"
+#include "Sound/SoundCue.h"
+#include "Sound/SoundWave.h"
+#include "ToolMenus.h"
+#include "Editor.h"
+#include "Styling/AppStyle.h"
+#include "Components/AudioComponent.h"
+#include "Widgets/Layout/SBox.h"
+#include "Widgets/Input/SButton.h"
+#include "Widgets/Images/SImage.h"
+
+#define LOCTEXT_NAMESPACE "AssetTypeActions"
+
+namespace UE::AudioEditor
+{
+	void StopSound()
+	{
+		GEditor->ResetPreviewAudioComponent();
+	}
+	
+	void PlaySound(USoundBase* Sound)
+	{
+		if ( Sound )
+		{
+			GEditor->PlayPreviewSound(Sound);
+		}
+		else
+		{
+			StopSound();
+		}
+	}
+	
+	bool IsSoundPlaying(USoundBase* Sound)
+	{
+		UAudioComponent* PreviewComp = GEditor->GetPreviewAudioComponent();
+		return PreviewComp && PreviewComp->Sound == Sound && PreviewComp->IsPlaying();
+	}
+	
+	bool IsSoundPlaying(const FAssetData& AssetData)
+	{
+		const UAudioComponent* PreviewComp = GEditor->GetPreviewAudioComponent();
+		if (PreviewComp && PreviewComp->Sound && PreviewComp->IsPlaying())
+		{
+			if (PreviewComp->Sound->GetFName() == AssetData.AssetName)
+			{
+				if (PreviewComp->Sound->GetOutermost()->GetFName() == AssetData.PackageName)
+				{
+					return true;
+				}
+			}
+		}
+	
+		return false;
+	}
+}
+
+EAssetCommandResult UAssetDefinition_SoundBase::ActivateAssets(const FAssetActivateArgs& ActivateArgs) const
+{
+	if ( ActivateArgs.ActivationMethod == EAssetActivationMethod::Previewed )
+	{
+		if ( USoundBase* TargetSound = ActivateArgs.LoadFirstValid<USoundBase>() )
+		{
+			UAudioComponent* PreviewComp = GEditor->GetPreviewAudioComponent();
+			if (PreviewComp && PreviewComp->IsPlaying())
+			{
+				// Already previewing a sound, if it is the target cue then stop it, otherwise play the new one
+				if (!TargetSound || PreviewComp->Sound == TargetSound)
+				{
+					UE::AudioEditor::StopSound();
+				}
+				else
+				{
+					UE::AudioEditor::PlaySound(TargetSound);
+				}
+			}
+			else
+			{
+				// Not already playing, play the target sound cue if it exists
+				UE::AudioEditor::PlaySound(TargetSound);
+			}
+			
+			return EAssetCommandResult::Handled;
+		}
+	}
+
+	return Super::ActivateAssets(ActivateArgs);
+}
+
+TSharedPtr<SWidget> UAssetDefinition_SoundBase::GetThumbnailOverlay(const FAssetData& InAssetData) const
+{
+	auto OnGetDisplayBrushLambda = [this, InAssetData]() -> const FSlateBrush*
+	{
+		if (UE::AudioEditor::IsSoundPlaying(InAssetData))
+		{
+			return FAppStyle::GetBrush("MediaAsset.AssetActions.Stop.Large");
+		}
+
+		return FAppStyle::GetBrush("MediaAsset.AssetActions.Play.Large");
+	};
+
+	auto OnClickedLambda = [this, InAssetData]() -> FReply
+	{
+		if (UE::AudioEditor::IsSoundPlaying(InAssetData))
+		{
+			UE::AudioEditor::StopSound();
+		}
+		else
+		{
+			// Load and play sound
+			UE::AudioEditor::PlaySound(Cast<USoundBase>(InAssetData.GetAsset()));
+		}
+		return FReply::Handled();
+	};
+
+	auto OnToolTipTextLambda = [this, InAssetData]() -> FText
+	{
+		if (UE::AudioEditor::IsSoundPlaying(InAssetData))
+		{
+			return LOCTEXT("Thumbnail_StopSoundToolTip", "Stop selected sound");
+		}
+
+		return LOCTEXT("Thumbnail_PlaySoundToolTip", "Play selected sound");
+	};
+
+	TSharedPtr<SBox> Box;
+	SAssignNew(Box, SBox)
+		.HAlign(HAlign_Center)
+		.VAlign(VAlign_Center)
+		.Padding(FMargin(2));
+
+	auto OnGetVisibilityLambda = [this, Box, InAssetData]() -> EVisibility
+	{
+		if (Box.IsValid() && (Box->IsHovered() || UE::AudioEditor::IsSoundPlaying(InAssetData)))
+		{
+			return EVisibility::Visible;
+		}
+
+		return EVisibility::Hidden;
+	};
+
+	TSharedPtr<SButton> Widget;
+	SAssignNew(Widget, SButton)
+		.ButtonStyle(FAppStyle::Get(), "HoverHintOnly")
+		.ToolTipText_Lambda(OnToolTipTextLambda)
+		.Cursor(EMouseCursor::Default) // The outer widget can specify a DragHand cursor, so we need to override that here
+		.ForegroundColor(FSlateColor::UseForeground())
+		.IsFocusable(false)
+		.OnClicked_Lambda(OnClickedLambda)
+		.Visibility_Lambda(OnGetVisibilityLambda)
+		[
+			SNew(SImage)
+			.Image_Lambda(OnGetDisplayBrushLambda)
+		];
+
+	Box->SetContent(Widget.ToSharedRef());
+	Box->SetVisibility(EVisibility::Visible);
+
+	return Box;
+}
+
+// Menu Extensions
+//--------------------------------------------------------------------
+
+namespace MenuExtension_SoundBase
+{
+	void ExecutePlaySound(const FToolMenuContext& InContext)
+	{
+		if (USoundBase* Sound = UContentBrowserAssetContextMenuContext::LoadSingleSelectedAsset<USoundBase>(InContext))
+		{
+			// Only play the first valid sound
+			UE::AudioEditor::PlaySound(Sound);
+		}
+	}
+
+	void ExecuteStopSound(const FToolMenuContext& InContext)
+	{
+		UE::AudioEditor::StopSound();
+	}
+
+	bool CanExecutePlayCommand(const FToolMenuContext& InContext)
+	{
+		if (const UContentBrowserAssetContextMenuContext* CBContext = UContentBrowserAssetContextMenuContext::FindContextWithAssets(InContext))
+		{
+			return CBContext->SelectedAssets.Num() == 1;
+		}
+
+		return false;
+	}
+	
+	ECheckBoxState IsActionCheckedMute(const FToolMenuContext& InContext)
+	{
+#if ENABLE_AUDIO_DEBUG
+		if (FAudioDeviceManager* ADM = GEditor->GetAudioDeviceManager())
+		{
+			if (const UContentBrowserAssetContextMenuContext* CBContext = UContentBrowserAssetContextMenuContext::FindContextWithAssets(InContext))
+			{
+				// If *any* of the selection are muted, show the tick box as ticked.
+				Audio::FAudioDebugger& Debugger = ADM->GetDebugger();
+
+				for (const FAssetData& SoundCueAsset : CBContext->GetSelectedAssetsOfType(USoundCue::StaticClass()))
+				{
+					if (Debugger.IsMuteSoundCue(SoundCueAsset.AssetName))
+					{
+						return ECheckBoxState::Checked;
+					}
+				}
+				
+				for (const FAssetData& SoundWaveAsset : CBContext->GetSelectedAssetsOfType(USoundWave::StaticClass()))
+				{
+					if (Debugger.IsMuteSoundWave(SoundWaveAsset.AssetName))
+					{
+						return ECheckBoxState::Checked;
+					}
+				}
+			}
+		}
+#endif
+		return ECheckBoxState::Unchecked;
+	}
+	
+	ECheckBoxState IsActionCheckedSolo(const FToolMenuContext& InContext)
+	{
+#if ENABLE_AUDIO_DEBUG
+		// If *any* of the selection are solod, show the tick box as ticked.
+		if (FAudioDeviceManager* ADM = GEditor->GetAudioDeviceManager())
+		{
+			if (const UContentBrowserAssetContextMenuContext* CBContext = UContentBrowserAssetContextMenuContext::FindContextWithAssets(InContext))
+			{
+				Audio::FAudioDebugger& Debugger = ADM->GetDebugger();
+
+				for (const FAssetData& SoundCueAsset : CBContext->GetSelectedAssetsOfType(USoundCue::StaticClass()))
+				{
+					if (Debugger.IsSoloSoundCue(SoundCueAsset.AssetName))
+					{
+						return ECheckBoxState::Checked;
+					}
+				}
+				
+				for (const FAssetData& SoundWaveAsset : CBContext->GetSelectedAssetsOfType(USoundWave::StaticClass()))
+				{
+					if (Debugger.IsSoloSoundWave(SoundWaveAsset.AssetName))
+                    {
+                    	return ECheckBoxState::Checked;
+                    }
+				}
+			}
+		}
+#endif
+		return ECheckBoxState::Unchecked;
+	}
+
+	void ExecuteMuteSound(const FToolMenuContext& InContext)
+	{
+#if ENABLE_AUDIO_DEBUG
+		if (FAudioDeviceManager* ADM = GEditor->GetAudioDeviceManager())
+		{
+			if (const UContentBrowserAssetContextMenuContext* CBContext = UContentBrowserAssetContextMenuContext::FindContextWithAssets(InContext))
+			{
+				Audio::FAudioDebugger& Debugger = ADM->GetDebugger();
+
+				// In a selection that consists of some already muted, toggle everything in the same direction,
+				// to avoid AB problem.
+				const bool bAnyMuted = IsActionCheckedMute(InContext) == ECheckBoxState::Checked;
+
+				for (const FAssetData& SoundCueAsset : CBContext->GetSelectedAssetsOfType(USoundCue::StaticClass()))
+				{
+					Debugger.SetMuteSoundCue(SoundCueAsset.AssetName, !bAnyMuted);
+				}
+				
+				for (const FAssetData& SoundWaveAsset : CBContext->GetSelectedAssetsOfType(USoundWave::StaticClass()))
+				{
+					Debugger.SetMuteSoundWave(SoundWaveAsset.AssetName, !bAnyMuted);
+				}
+			}
+		}
+#endif
+	}
+
+	void ExecuteSoloSound(const FToolMenuContext& InContext)
+	{
+#if ENABLE_AUDIO_DEBUG
+		if (FAudioDeviceManager* ADM = GEditor->GetAudioDeviceManager())
+		{
+			if (const UContentBrowserAssetContextMenuContext* CBContext = UContentBrowserAssetContextMenuContext::FindContextWithAssets(InContext))
+			{
+				Audio::FAudioDebugger& Debugger = ADM->GetDebugger();
+
+				// In a selection that consists of some already soloed, toggle everything in the same direction,
+				// to avoid AB problem.
+
+				const bool bAnySoloed = IsActionCheckedSolo(InContext) == ECheckBoxState::Checked;
+
+				for (const FAssetData& SoundCueAsset : CBContext->GetSelectedAssetsOfType(USoundCue::StaticClass()))
+				{
+					Debugger.SetSoloSoundCue(SoundCueAsset.AssetName, !bAnySoloed);
+				}
+				
+				for (const FAssetData& SoundWaveAsset : CBContext->GetSelectedAssetsOfType(USoundWave::StaticClass()))
+				{
+					Debugger.SetSoloSoundWave(SoundWaveAsset.AssetName, !bAnySoloed);
+				}
+			}
+		}
+#endif
+	}	
+
+	bool CanExecuteMuteCommand(const FToolMenuContext& InContext)
+	{
+#if ENABLE_AUDIO_DEBUG
+		if (FAudioDeviceManager* ADM = GEditor->GetAudioDeviceManager())
+		{
+			if (const UContentBrowserAssetContextMenuContext* CBContext = UContentBrowserAssetContextMenuContext::FindContextWithAssets(InContext))
+			{
+				// Allow muting if we're not Soloing.
+				Audio::FAudioDebugger& Debugger = ADM->GetDebugger();
+
+				for (const FAssetData& SoundCueAsset : CBContext->GetSelectedAssetsOfType(USoundCue::StaticClass()))
+				{
+					if (Debugger.IsSoloSoundCue(SoundCueAsset.AssetName))
+					{
+						return false;
+					}
+				}
+				
+				for (const FAssetData& SoundWaveAsset : CBContext->GetSelectedAssetsOfType(USoundWave::StaticClass()))
+				{
+					if (Debugger.IsSoloSoundWave(SoundWaveAsset.AssetName))
+					{
+						return false;
+					}
+				}
+
+				// Ok.
+				return true;
+			}
+		}
+#endif
+		return false;
+	}
+
+	bool CanExecuteSoloCommand(const FToolMenuContext& InContext)
+	{
+#if ENABLE_AUDIO_DEBUG
+		if (FAudioDeviceManager* ADM = GEditor->GetAudioDeviceManager())
+		{
+			if (const UContentBrowserAssetContextMenuContext* CBContext = UContentBrowserAssetContextMenuContext::FindContextWithAssets(InContext))
+			{
+				// Allow Soloing if we're not Muting.
+				Audio::FAudioDebugger& Debugger = ADM->GetDebugger();
+			
+				for (const FAssetData& SoundCueAsset : CBContext->GetSelectedAssetsOfType(USoundCue::StaticClass()))
+				{
+					if (Debugger.IsMuteSoundCue(SoundCueAsset.AssetName))
+					{
+						return false;
+					}
+				}
+            
+				for (const FAssetData& SoundWaveAsset : CBContext->GetSelectedAssetsOfType(USoundWave::StaticClass()))
+				{
+					if (Debugger.IsMuteSoundWave(SoundWaveAsset.AssetName))
+					{
+						return false;
+					}
+				}
+
+				// Ok.
+				return true;
+			}
+		}
+#endif
+		return false;
+	}
+	
+	static FDelayedAutoRegisterHelper DelayedAutoRegister(EDelayedRegisterRunPhase::EndOfEngineInit, []{ 
+		UToolMenus::RegisterStartupCallback(FSimpleMulticastDelegate::FDelegate::CreateLambda([]()
+		{
+			FToolMenuOwnerScoped OwnerScoped(UE_MODULE_NAME);
+			UToolMenu* Menu = UE::ContentBrowser::ExtendToolMenu_AssetContextMenu(USoundBase::StaticClass());
+		
+			FToolMenuSection& Section = Menu->FindOrAddSection("GetAssetActions");
+			Section.AddDynamicEntry(NAME_None, FNewToolMenuSectionDelegate::CreateLambda([](FToolMenuSection& InSection)
+			{
+				if (const UContentBrowserAssetContextMenuContext* Context = UContentBrowserAssetContextMenuContext::FindContextWithAssets(InSection))
+				{
+					{
+						const TAttribute<FText> Label = LOCTEXT("Sound_PlaySound", "Play");
+						const TAttribute<FText> ToolTip = LOCTEXT("Sound_PlaySoundTooltip", "Plays the selected sound.");
+						const FSlateIcon Icon = FSlateIcon(FAppStyle::GetAppStyleSetName(), "MediaAsset.AssetActions.Play.Small");
+						
+						FToolUIAction UIAction;
+						UIAction.ExecuteAction = FToolMenuExecuteAction::CreateStatic(&ExecutePlaySound);
+						UIAction.CanExecuteAction = FToolMenuCanExecuteAction::CreateStatic(&CanExecutePlayCommand);
+						InSection.AddMenuEntry("Sound_PlaySound", Label, ToolTip, Icon, UIAction);
+					}
+					{
+						const TAttribute<FText> Label = LOCTEXT("Sound_StopSound", "Stop");
+						const TAttribute<FText> ToolTip = LOCTEXT("Sound_StopSoundTooltip", "Stops the selected sounds.");
+						const FSlateIcon Icon = FSlateIcon(FAppStyle::GetAppStyleSetName(), "MediaAsset.AssetActions.Stop.Small");
+						
+						FToolUIAction UIAction;
+						UIAction.ExecuteAction = FToolMenuExecuteAction::CreateStatic(&ExecuteStopSound);
+						InSection.AddMenuEntry("Sound_StopSound", Label, ToolTip, Icon, UIAction);
+					}
+					{
+						const TAttribute<FText> Label = LOCTEXT("Sound_MuteSound", "Mute");
+						const TAttribute<FText> ToolTip = LOCTEXT("Sound_MuteSoundTooltip", "Mutes the selected sounds.");
+						const FSlateIcon Icon = FSlateIcon(FAppStyle::GetAppStyleSetName(), "MediaAsset.AssetActions.Mute.Small");
+		
+						FToolUIAction UIAction;
+						UIAction.ExecuteAction = FToolMenuExecuteAction::CreateStatic(&ExecuteMuteSound);
+						UIAction.CanExecuteAction = FToolMenuCanExecuteAction::CreateStatic(&CanExecuteMuteCommand);
+						UIAction.GetActionCheckState = FToolMenuGetActionCheckState::CreateStatic(&IsActionCheckedMute);
+						InSection.AddMenuEntry("Sound_SoundMute", Label, ToolTip, Icon, UIAction, EUserInterfaceActionType::ToggleButton);
+					}
+					{
+						const TAttribute<FText> Label = LOCTEXT("Sound_SoloSound", "Solo");
+						const TAttribute<FText> ToolTip = LOCTEXT("Sound_SoloSoundTooltip", "Solos the selected sounds.");
+						const FSlateIcon Icon = FSlateIcon(FAppStyle::GetAppStyleSetName(), "MediaAsset.AssetActions.Solo.Small");
+
+						FToolUIAction UIAction;
+						UIAction.ExecuteAction = FToolMenuExecuteAction::CreateStatic(&ExecuteSoloSound);
+						UIAction.CanExecuteAction = FToolMenuCanExecuteAction::CreateStatic(&CanExecuteSoloCommand);
+						UIAction.GetActionCheckState = FToolMenuGetActionCheckState::CreateStatic(&IsActionCheckedSolo);
+						InSection.AddMenuEntry("Sound_StopSolo", Label, ToolTip, Icon, UIAction, EUserInterfaceActionType::ToggleButton);
+					}
+				}
+			}));
+		}));
+	});
+}
+
+#undef LOCTEXT_NAMESPACE
diff --git a/Engine/Source/Editor/AudioEditor/Private/AssetTypeActions/AssetDefinition_SoundCue.cpp b/Engine/Source/Editor/AudioEditor/Private/AssetTypeActions/AssetDefinition_SoundCue.cpp
new file mode 100644
--- /dev/null
+++ b/Engine/Source/Editor/AudioEditor/Private/AssetTypeActions/AssetDefinition_SoundCue.cpp
@@ -0,0 +1,176 @@
+// Copyright Epic Games, Inc. All Rights Reserved.
+
+#include "AssetTypeActions/AssetDefinition_SoundCue.h"
+
+#include "AssetToolsModule.h"
+#include "AudioEditorModule.h"
+#include "AudioEditorSettings.h"
+#include "Containers/Map.h"
+#include "Containers/UnrealString.h"
+#include "ContentBrowserModule.h"
+#include "Delegates/Delegate.h"
+#include "Factories/SoundAttenuationFactory.h"
+#include "Framework/Commands/UIAction.h"
+#include "HAL/Platform.h"
+#include "HAL/PlatformCrt.h"
+#include "IAssetTools.h"
+#include "IContentBrowserSingleton.h"
+#include "Misc/PackageName.h"
+#include "Modules/ModuleManager.h"
+#include "Sound/SoundAttenuation.h"
+#include "Sound/SoundCue.h"
+#include "Templates/Casts.h"
+#include "Templates/ChooseClass.h"
+#include "Textures/SlateIcon.h"
+#include "ToolMenuSection.h"
+#include "Toolkits/IToolkit.h"
+#include "UObject/Object.h"
+#include "UObject/ObjectPtr.h"
+#include "UObject/UObjectGlobals.h"
+#include "ContentBrowserMenuContexts.h"
+#include "ToolMenus.h"
+
+#include <utility>
+
+class IToolkitHost;
+class UClass;
+
+#define LOCTEXT_NAMESPACE "AssetTypeActions"
+
+TConstArrayView<FAssetCategoryPath> UAssetDefinition_SoundCue::GetAssetCategories() const
+{
+	/*
+	static const auto Categories = {
+		EAssetCategoryPaths::Audio,
+		EAssetCategoryPaths::Audio / LOCTEXT("AssetSoundCueSubMenu", "Source")
+	};
+	*/
+
+	static const auto Pinned_Categories = { EAssetCategoryPaths::Audio };
+	static const auto Categories = { EAssetCategoryPaths::Audio / LOCTEXT("AssetSoundCueSubMenu", "Source") };
+	
+	if (GetDefault<UAudioEditorSettings>()->bPinSoundCueInAssetMenu)
+	{
+		return Pinned_Categories;
+	}
+	
+	return Categories;
+}
+
+EAssetCommandResult UAssetDefinition_SoundCue::OpenAssets(const FAssetOpenArgs& OpenArgs) const
+{
+	IAudioEditorModule& AudioEditorModule = FModuleManager::LoadModuleChecked<IAudioEditorModule>( "AudioEditor" );
+	
+	for (USoundCue* SoundCue : OpenArgs.LoadObjects<USoundCue>())
+	{
+		AudioEditorModule.CreateSoundCueEditor(OpenArgs.GetToolkitMode(), OpenArgs.ToolkitHost, SoundCue);
+	}
+
+	return EAssetCommandResult::Handled;
+}
+
+// Menu Extensions
+//--------------------------------------------------------------------
+
+namespace MenuExtension_SoundCue
+{
+	bool CanExecuteConsolidateAttenuation(const FToolMenuContext& MenuContext)
+	{
+		return UContentBrowserAssetContextMenuContext::GetNumAssetsSelected(MenuContext) > 1;
+	}
+
+	void ExecuteConsolidateAttenuation(const FToolMenuContext& MenuContext)
+	{
+		TMap<FSoundAttenuationSettings*, TArray<USoundCue*>> UnmatchedAttenuations;
+
+		if (const UContentBrowserAssetContextMenuContext* Context = UContentBrowserAssetContextMenuContext::FindContextWithAssets(MenuContext))
+		{
+			for (USoundCue* SoundCue : Context->LoadSelectedObjects<USoundCue>())
+			{
+				bool bFound = false;
+				if ( SoundCue && SoundCue->bOverrideAttenuation )
+				{
+					for (auto UnmatchedIt = UnmatchedAttenuations.CreateIterator(); UnmatchedIt; ++UnmatchedIt)
+					{
+						// Found attenuation settings to consolidate together
+						if (SoundCue->AttenuationOverrides == *UnmatchedIt.Key())
+						{
+							UnmatchedIt.Value().Add(SoundCue);
+							bFound = true;
+							break;
+						}
+					}
+					if (!bFound)
+					{
+						UnmatchedAttenuations.FindOrAdd(&SoundCue->AttenuationOverrides).Add(SoundCue);
+					}
+				}
+			}
+		}
+
+		if (UnmatchedAttenuations.Num() > 0)
+		{
+			FString DefaultSuffix;
+			TArray<UObject*> ObjectsToSync;
+
+			FAssetToolsModule& AssetToolsModule = FModuleManager::GetModuleChecked<FAssetToolsModule>("AssetTools");
+			USoundAttenuationFactory* Factory = NewObject<USoundAttenuationFactory>();
+
+			for (auto UnmatchedIt = UnmatchedAttenuations.CreateConstIterator(); UnmatchedIt; ++UnmatchedIt)
+			{
+				if (UnmatchedIt.Value().Num() > 1)
+				{
+					FString Name;
+					FString PackageName;
+					AssetToolsModule.Get().CreateUniqueAssetName("/Game/Sounds/SoundAttenuations/SharedAttenuation", DefaultSuffix, PackageName, Name);
+
+					USoundAttenuation* SoundAttenuation = Cast<USoundAttenuation>(AssetToolsModule.Get().CreateAsset(Name, FPackageName::GetLongPackagePath(PackageName), USoundAttenuation::StaticClass(), Factory));
+					if (SoundAttenuation)
+					{
+						SoundAttenuation->Attenuation = *UnmatchedIt.Key();
+
+						for (int32 SoundCueIndex = 0; SoundCueIndex < UnmatchedIt.Value().Num(); ++SoundCueIndex)
+						{
+							USoundCue* SoundCue = UnmatchedIt.Value()[SoundCueIndex];
+							SoundCue->bOverrideAttenuation = false;
+							SoundCue->AttenuationSettings = SoundAttenuation;
+							SoundCue->MarkPackageDirty();
+						}
+					}
+				}
+			}
+
+			if ( ObjectsToSync.Num() > 0 )
+			{
+				FContentBrowserModule& ContentBrowserModule = FModuleManager::LoadModuleChecked<FContentBrowserModule>("ContentBrowser");
+				ContentBrowserModule.Get().SyncBrowserToAssets(ObjectsToSync);
+			}
+		}
+	}
+	
+	static FDelayedAutoRegisterHelper DelayedAutoRegister(EDelayedRegisterRunPhase::EndOfEngineInit, []{ 
+		UToolMenus::RegisterStartupCallback(FSimpleMulticastDelegate::FDelegate::CreateLambda([]()
+		{
+			FToolMenuOwnerScoped OwnerScoped(UE_MODULE_NAME);
+
+			{
+				UToolMenu* Menu = UE::ContentBrowser::ExtendToolMenu_AssetContextMenu(USoundCue::StaticClass());
+		
+				FToolMenuSection& Section = Menu->FindOrAddSection("GetAssetActions");
+				Section.AddDynamicEntry(NAME_None, FNewToolMenuSectionDelegate::CreateLambda([](FToolMenuSection& InSection)
+				{
+					const TAttribute<FText> Label = LOCTEXT("SoundCue_ConsolidateAttenuation", "Consolidate Attenuation");
+					const TAttribute<FText> ToolTip = LOCTEXT("SoundCue_ConsolidateAttenuationTooltip", "Creates shared attenuation packages for sound cues with identical override attenuation settings.");
+					const FSlateIcon Icon = FSlateIcon(FAppStyle::GetAppStyleSetName(), "ClassIcon.SoundCue");
+
+					FToolUIAction UIAction;
+                    UIAction.ExecuteAction = FToolMenuExecuteAction::CreateStatic(&ExecuteConsolidateAttenuation);
+                    UIAction.CanExecuteAction = FToolMenuCanExecuteAction::CreateStatic(&CanExecuteConsolidateAttenuation);
+					InSection.AddMenuEntry("SoundCue_ConsolidateAttenuation", Label, ToolTip, Icon, UIAction);
+				}));
+			}
+		}));
+	});
+}
+
+#undef LOCTEXT_NAMESPACE
diff --git a/Engine/Source/Editor/AudioEditor/Private/AssetTypeActions/AssetDefinition_SoundCue.h b/Engine/Source/Editor/AudioEditor/Private/AssetTypeActions/AssetDefinition_SoundCue.h
new file mode 100644
--- /dev/null
+++ b/Engine/Source/Editor/AudioEditor/Private/AssetTypeActions/AssetDefinition_SoundCue.h
@@ -0,0 +1,24 @@
+// Copyright Epic Games, Inc. All Rights Reserved.
+
+#pragma once
+
+#include "CoreMinimal.h"
+#include "Sound/SoundCue.h"
+#include "AssetTypeActions/AssetDefinition_SoundBase.h"
+
+#include "AssetDefinition_SoundCue.generated.h"
+
+UCLASS()
+class UAssetDefinition_SoundCue : public UAssetDefinition_SoundBase
+{
+	GENERATED_BODY()
+
+public:
+	// UAssetDefinition Begin
+	virtual FText GetAssetDisplayName() const override { return NSLOCTEXT("AssetTypeActions", "AssetTypeActions_SoundCue", "Sound Cue"); }
+	virtual FLinearColor GetAssetColor() const override { return FLinearColor(FColor(0, 175, 255)); }
+	virtual TSoftClassPtr<UObject> GetAssetClass() const override { return USoundCue::StaticClass(); }
+	virtual TConstArrayView<FAssetCategoryPath> GetAssetCategories() const override;
+	virtual EAssetCommandResult OpenAssets(const FAssetOpenArgs& OpenArgs) const override;
+	// UAssetDefinition End
+};
diff --git a/Engine/Source/Editor/AudioEditor/Private/AssetTypeActions/AssetDefinition_SoundSourceBus.h b/Engine/Source/Editor/AudioEditor/Private/AssetTypeActions/AssetDefinition_SoundSourceBus.h
new file mode 100644
--- /dev/null
+++ b/Engine/Source/Editor/AudioEditor/Private/AssetTypeActions/AssetDefinition_SoundSourceBus.h
@@ -0,0 +1,28 @@
+// Copyright Epic Games, Inc. All Rights Reserved.
+
+#pragma once
+
+#include "CoreMinimal.h"
+#include "Sound/SoundSourceBus.h"
+#include "AssetTypeActions/AssetDefinition_SoundBase.h"
+
+#include "AssetDefinition_SoundSourceBus.generated.h"
+
+UCLASS()
+class UAssetDefinition_SoundSourceBus : public UAssetDefinition_SoundBase
+{
+	GENERATED_BODY()
+
+public:
+	// UAssetDefinition Begin
+	virtual FText GetAssetDisplayName() const override { return NSLOCTEXT("AssetTypeActions", "AssetTypeActions_SoundSourceBus", "Source Bus"); }
+	virtual FLinearColor GetAssetColor() const override { return FLinearColor(FColor(212, 97, 85)); }
+	virtual TSoftClassPtr<UObject> GetAssetClass() const override { return USoundSourceBus::StaticClass(); }
+	virtual TConstArrayView<FAssetCategoryPath> GetAssetCategories() const override
+	{
+		static const auto Categories = { FAssetCategoryPath(EAssetCategoryPaths::Audio, NSLOCTEXT("AssetTypeActions", "AssetSoundSourceSubMenu", "Source")) };
+		return Categories;
+	}
+	virtual TSharedPtr<SWidget> GetThumbnailOverlay(const FAssetData& InAssetData) const override { return nullptr; }
+	// UAssetDefinition End
+};
diff --git a/Engine/Source/Editor/AudioEditor/Private/AssetTypeActions/AssetTypeActions_SoundBase.cpp b/Engine/Source/Editor/AudioEditor/Private/AssetTypeActions/AssetTypeActions_SoundBase.cpp
--- a/Engine/Source/Editor/AudioEditor/Private/AssetTypeActions/AssetTypeActions_SoundBase.cpp
+++ b/Engine/Source/Editor/AudioEditor/Private/AssetTypeActions/AssetTypeActions_SoundBase.cpp
@@ -21,59 +21,6 @@ UClass* FAssetTypeActions_SoundBase::GetSupportedClass() const
 	return USoundBase::StaticClass();
 }
 
-void FAssetTypeActions_SoundBase::GetActions(const TArray<UObject*>& InObjects, FToolMenuSection& Section)
-{
-	auto Sounds = GetTypedWeakObjectPtrs<USoundBase>(InObjects);
-
-	Section.AddMenuEntry(
-		"Sound_PlaySound",
-		LOCTEXT("Sound_PlaySound", "Play"),
-		LOCTEXT("Sound_PlaySoundTooltip", "Plays the selected sound."),
-		FSlateIcon(FAppStyle::GetAppStyleSetName(), "MediaAsset.AssetActions.Play.Small"),
-		FUIAction(
-			FExecuteAction::CreateSP( this, &FAssetTypeActions_SoundBase::ExecutePlaySound, Sounds ),
-			FCanExecuteAction::CreateSP( this, &FAssetTypeActions_SoundBase::CanExecutePlayCommand, Sounds )
-			)
-		);
-
-	Section.AddMenuEntry(
-		"Sound_StopSound",
-		LOCTEXT("Sound_StopSound", "Stop"),
-		LOCTEXT("Sound_StopSoundTooltip", "Stops the selected sounds."),
-		FSlateIcon(FAppStyle::GetAppStyleSetName(), "MediaAsset.AssetActions.Stop.Small"),
-		FUIAction(
-			FExecuteAction::CreateSP( this, &FAssetTypeActions_SoundBase::ExecuteStopSound, Sounds ),
-			FCanExecuteAction()
-			)
-		);
-
-	Section.AddMenuEntry(
-		"Sound_SoundMute",
-		LOCTEXT("Sound_MuteSound", "Mute"),
-		LOCTEXT("Sound_MuteSoundTooltip", "Mutes the selected sounds."),
-		FSlateIcon(FAppStyle::GetAppStyleSetName(), "MediaAsset.AssetActions.Mute.Small"),
-		FUIAction(
-			FExecuteAction::CreateSP(this, &FAssetTypeActions_SoundBase::ExecuteMuteSound, Sounds),
-			FCanExecuteAction::CreateSP(this, &FAssetTypeActions_SoundBase::CanExecuteMuteCommand, Sounds),
-			FIsActionChecked::CreateSP(this, &FAssetTypeActions_SoundBase::IsActionCheckedMute, Sounds)
-		),
-		EUserInterfaceActionType::ToggleButton
-	);
-
-	Section.AddMenuEntry(
-		"Sound_StopSolo",
-		LOCTEXT("Sound_SoloSound", "Solo"),
-		LOCTEXT("Sound_SoloSoundTooltip", "Solos the selected sounds."),
-		FSlateIcon(FAppStyle::GetAppStyleSetName(), "MediaAsset.AssetActions.Solo.Small"),
-		FUIAction(
-			FExecuteAction::CreateSP(this, &FAssetTypeActions_SoundBase::ExecuteSoloSound, Sounds),
-			FCanExecuteAction::CreateSP(this, &FAssetTypeActions_SoundBase::CanExecuteSoloCommand, Sounds),
-			FIsActionChecked::CreateSP(this, &FAssetTypeActions_SoundBase::IsActionCheckedSolo, Sounds)
-		),
-		EUserInterfaceActionType::ToggleButton
-	);
-}
-
 bool FAssetTypeActions_SoundBase::AssetsActivatedOverride(const TArray<UObject*>& InObjects, EAssetTypeActivationMethod::Type ActivationType)
 {
 	if (ActivationType == EAssetTypeActivationMethod::Previewed)
@@ -244,183 +191,4 @@ TSharedPtr<SWidget> FAssetTypeActions_SoundBase::GetThumbnailOverlay(const FAsse
 	return Box;
 }
 
-bool FAssetTypeActions_SoundBase::CanExecutePlayCommand(TArray<TWeakObjectPtr<USoundBase>> Objects) const
-{
-	return Objects.Num() == 1;
-}
-
-void FAssetTypeActions_SoundBase::ExecuteMuteSound(TArray<TWeakObjectPtr<USoundBase>> Objects) const
-{	
-#if ENABLE_AUDIO_DEBUG
-	if (FAudioDeviceManager* ADM = GEditor->GetAudioDeviceManager())
-	{
-		Audio::FAudioDebugger& Debugger = ADM->GetDebugger();
-
-		// In a selection that consists of some already muted, toggle everything in the same direction,
-		// to avoid AB problem.
-
-		bool bAnyMuted = IsActionCheckedMute(Objects);
-		for (TWeakObjectPtr<USoundBase> SoundBase : Objects)
-		{
-			if (USoundCue* Cue = Cast<USoundCue>(SoundBase.Get()))
-			{
-				Debugger.SetMuteSoundCue(Cue->GetFName(), !bAnyMuted);
-			}
-			else if (USoundWave* Wave = Cast<USoundWave>(SoundBase.Get()))
-			{
-				Debugger.SetMuteSoundWave(Wave->GetFName(), !bAnyMuted);
-			}
-		}
-	}
-#endif
-}
-
-void FAssetTypeActions_SoundBase::ExecuteSoloSound(TArray<TWeakObjectPtr<USoundBase>> Objects) const
-{
-#if ENABLE_AUDIO_DEBUG
-	if (FAudioDeviceManager* ADM = GEditor->GetAudioDeviceManager())
-	{
-		Audio::FAudioDebugger& Debugger = ADM->GetDebugger();
-
-		// In a selection that consists of some already soloed, toggle everything in the same direction,
-		// to avoid AB problem.
-
-		bool bAnySoloed = IsActionCheckedSolo(Objects);
-		for (TWeakObjectPtr<USoundBase> SoundBase : Objects)
-		{
-			if (USoundCue* Cue = Cast<USoundCue>(SoundBase.Get()))
-			{
-				Debugger.SetSoloSoundCue(Cue->GetFName(), !bAnySoloed);
-			}
-			else if (USoundWave* Wave = Cast<USoundWave>(SoundBase.Get()))
-			{
-				Debugger.SetSoloSoundWave(Wave->GetFName(), !bAnySoloed);
-			}
-		}
-	}
-	#endif
-}
-
-bool FAssetTypeActions_SoundBase::IsActionCheckedMute(TArray<TWeakObjectPtr<USoundBase>> Objects) const
-{
-#if ENABLE_AUDIO_DEBUG
-	if (FAudioDeviceManager* ADM = GEditor->GetAudioDeviceManager())
-	{
-		// If *any* of the selection are muted, show the tick box as ticked.
-		Audio::FAudioDebugger& Debugger = ADM->GetDebugger();
-		for (TWeakObjectPtr<USoundBase> SoundBase : Objects)
-		{
-			if (USoundCue* Cue = Cast<USoundCue>(SoundBase.Get()))
-			{
-				if (Debugger.IsMuteSoundCue(Cue->GetFName()))
-				{
-					return true;
-				}
-			}
-			else if (USoundWave* Wave = Cast<USoundWave>(SoundBase.Get()))
-			{
-				if (Debugger.IsMuteSoundWave(Wave->GetFName()))
-				{
-					return true;
-				}
-			}
-		}
-	}
-#endif
-	return false;
-}
-
-bool FAssetTypeActions_SoundBase::IsActionCheckedSolo(TArray<TWeakObjectPtr<USoundBase>> Objects) const
-{
-#if ENABLE_AUDIO_DEBUG
-	// If *any* of the selection are solod, show the tick box as ticked.
-	if (FAudioDeviceManager* ADM = GEditor->GetAudioDeviceManager())
-	{
-		Audio::FAudioDebugger& Debugger = ADM->GetDebugger();
-		for (TWeakObjectPtr<USoundBase> SoundBase : Objects)
-		{
-			if (USoundCue* Cue = Cast<USoundCue>(SoundBase.Get()))
-			{
-				if (Debugger.IsSoloSoundCue(Cue->GetFName()))
-				{
-					return true;
-				}
-			}
-			else if (USoundWave* Wave = Cast<USoundWave>(SoundBase.Get()))
-			{
-				if (Debugger.IsSoloSoundWave(Wave->GetFName()))
-				{
-					return true;
-				}
-			}
-		}
-	}
-#endif
-	return false;
-}
-
-bool FAssetTypeActions_SoundBase::CanExecuteMuteCommand(TArray<TWeakObjectPtr<USoundBase>> Objects) const
-{
-#if ENABLE_AUDIO_DEBUG
-	if (FAudioDeviceManager* ADM = GEditor->GetAudioDeviceManager())
-	{
-		// Allow muting if we're not Soloing.
-		Audio::FAudioDebugger& Debugger = ADM->GetDebugger();
-		for (TWeakObjectPtr<USoundBase> SoundBase : Objects)
-		{
-			if (USoundCue* Cue = Cast<USoundCue>(SoundBase.Get()))
-			{
-				if (Debugger.IsSoloSoundCue(Cue->GetFName()))
-				{
-					return false;
-				}
-			}
-			else if (USoundWave* Wave = Cast<USoundWave>(SoundBase.Get()))
-			{
-				if (Debugger.IsSoloSoundWave(Wave->GetFName()))
-				{
-					return false;
-				}
-			}
-		}
-
-		// Ok.
-		return true;
-	}
-#endif
-	return false;
-}
-
-bool FAssetTypeActions_SoundBase::CanExecuteSoloCommand(TArray<TWeakObjectPtr<USoundBase>> Objects) const
-{	
-#if ENABLE_AUDIO_DEBUG
-	if (FAudioDeviceManager* ADM = GEditor->GetAudioDeviceManager())
-	{
-		// Allow Soloing if we're not Muting.
-		Audio::FAudioDebugger& Debugger = ADM->GetDebugger();
-		for (TWeakObjectPtr<USoundBase> SoundBase : Objects)
-		{
-			if (USoundCue* Cue = Cast<USoundCue>(SoundBase.Get()))
-			{
-				if (Debugger.IsMuteSoundCue(Cue->GetFName()))
-				{
-					return false;
-				}
-			}
-			else if (USoundWave* Wave = Cast<USoundWave>(SoundBase.Get()))
-			{
-				if (Debugger.IsMuteSoundWave(Wave->GetFName()))
-				{
-					return false;
-				}
-			}
-		}
-
-		// Ok.
-		return true;
-	}
-#endif
-	return false;
-}
-
 #undef LOCTEXT_NAMESPACE
diff --git a/Engine/Source/Editor/AudioEditor/Private/AssetTypeActions/AssetTypeActions_SoundCue.cpp b/Engine/Source/Editor/AudioEditor/Private/AssetTypeActions/AssetTypeActions_SoundCue.cpp
deleted file mode 100644
diff --git a/Engine/Source/Editor/AudioEditor/Private/AssetTypeActions/AssetTypeActions_SoundCue.h b/Engine/Source/Editor/AudioEditor/Private/AssetTypeActions/AssetTypeActions_SoundCue.h
deleted file mode 100644
diff --git a/Engine/Source/Editor/AudioEditor/Private/AssetTypeActions/AssetTypeActions_SoundSourceBus.cpp b/Engine/Source/Editor/AudioEditor/Private/AssetTypeActions/AssetTypeActions_SoundSourceBus.cpp
deleted file mode 100644
diff --git a/Engine/Source/Editor/AudioEditor/Private/AssetTypeActions/AssetTypeActions_SoundSourceBus.h b/Engine/Source/Editor/AudioEditor/Private/AssetTypeActions/AssetTypeActions_SoundSourceBus.h
deleted file mode 100644
diff --git a/Engine/Source/Editor/AudioEditor/Private/AudioEditorModule.cpp b/Engine/Source/Editor/AudioEditor/Private/AudioEditorModule.cpp
--- a/Engine/Source/Editor/AudioEditor/Private/AudioEditorModule.cpp
+++ b/Engine/Source/Editor/AudioEditor/Private/AudioEditorModule.cpp
@@ -4,12 +4,9 @@
 
 #include "AssetToolsModule.h"
 #include "AssetTypeActions/AssetTypeActions_DialogueWave.h"
-#include "AssetTypeActions/AssetTypeActions_SoundBase.h"
 #include "AssetTypeActions/AssetTypeActions_SoundClass.h"
-#include "AssetTypeActions/AssetTypeActions_SoundCue.h"
 #include "AssetTypeActions/AssetTypeActions_SoundEffectPreset.h"
 #include "AssetTypeActions/AssetTypeActions_SoundWave.h"
-#include "AssetTypeActions/AssetTypeActions_SoundSourceBus.h"
 #include "AssetTypeActions/AssetTypeActions_SoundSubmix.h"
 #include "ClassTemplateEditorSubsystem.h"
 #include "Components/SynthComponent.h"
@@ -181,9 +178,7 @@ public:
 		IAssetTools& AssetTools = FModuleManager::LoadModuleChecked<FAssetToolsModule>("AssetTools").Get();
 
 		AssetTools.RegisterAssetTypeActions(MakeShared<FAssetTypeActions_DialogueWave>());
-		AssetTools.RegisterAssetTypeActions(MakeShared<FAssetTypeActions_SoundBase>());
 		AssetTools.RegisterAssetTypeActions(MakeShared<FAssetTypeActions_SoundClass>());
-		AssetTools.RegisterAssetTypeActions(MakeShared<FAssetTypeActions_SoundCue>());
 		AssetTools.RegisterAssetTypeActions(MakeShared<FAssetTypeActions_SoundWave>());
 	}
 
@@ -205,7 +200,6 @@ public:
 			AssetTools.RegisterAssetTypeActions(MakeShared<FAssetTypeActions_SoundEffectSubmixPreset>());
 			AssetTools.RegisterAssetTypeActions(MakeShared<FAssetTypeActions_SoundEffectSourcePreset>());
 			AssetTools.RegisterAssetTypeActions(MakeShared<FAssetTypeActions_SoundEffectSourcePresetChain>());
-			AssetTools.RegisterAssetTypeActions(MakeShared<FAssetTypeActions_SoundSourceBus>());
 		}
 	}
 
diff --git a/Engine/Source/Editor/AudioEditor/Public/AssetTypeActions/AssetDefinition_SoundBase.h b/Engine/Source/Editor/AudioEditor/Public/AssetTypeActions/AssetDefinition_SoundBase.h
new file mode 100644
--- /dev/null
+++ b/Engine/Source/Editor/AudioEditor/Public/AssetTypeActions/AssetDefinition_SoundBase.h
@@ -0,0 +1,39 @@
+// Copyright Epic Games, Inc. All Rights Reserved.
+
+#pragma once
+
+#include "CoreMinimal.h"
+#include "Sound/SoundBase.h"
+#include "AssetDefinitionDefault.h"
+
+#include "AssetDefinition_SoundBase.generated.h"
+
+namespace UE::AudioEditor
+{
+	AUDIOEDITOR_API void StopSound();
+	
+	AUDIOEDITOR_API void PlaySound(USoundBase* Sound);
+	
+	AUDIOEDITOR_API bool IsSoundPlaying(USoundBase* Sound);
+	AUDIOEDITOR_API bool IsSoundPlaying(const FAssetData& AssetData);
+}
+
+UCLASS()
+class AUDIOEDITOR_API UAssetDefinition_SoundBase : public UAssetDefinitionDefault
+{
+	GENERATED_BODY()
+
+public:
+	// UAssetDefinition Begin
+	virtual FText GetAssetDisplayName() const override { return NSLOCTEXT("AssetTypeActions", "AssetTypeActions_SoundBase", "Sound Base"); }
+	virtual FLinearColor GetAssetColor() const override { return FLinearColor(FColor(97, 85, 212)); }
+	virtual TSoftClassPtr<UObject> GetAssetClass() const override { return USoundBase::StaticClass(); }
+	virtual TConstArrayView<FAssetCategoryPath> GetAssetCategories() const override
+	{
+		static const auto Categories = { EAssetCategoryPaths::Audio };
+		return Categories;
+	}
+	virtual TSharedPtr<SWidget> GetThumbnailOverlay(const FAssetData& InAssetData) const override;
+	virtual EAssetCommandResult ActivateAssets(const FAssetActivateArgs& ActivateArgs) const override;
+	// UAssetDefinition End
+};
diff --git a/Engine/Source/Editor/AudioEditor/Public/AssetTypeActions/AssetTypeActions_SoundBase.h b/Engine/Source/Editor/AudioEditor/Public/AssetTypeActions/AssetTypeActions_SoundBase.h
--- a/Engine/Source/Editor/AudioEditor/Public/AssetTypeActions/AssetTypeActions_SoundBase.h
+++ b/Engine/Source/Editor/AudioEditor/Public/AssetTypeActions/AssetTypeActions_SoundBase.h
@@ -20,15 +20,16 @@ class UObject;
 class USoundBase;
 struct FAssetData;
 
-PRAGMA_DISABLE_DEPRECATION_WARNINGS
-class AUDIOEDITOR_API FAssetTypeActions_SoundBase : public FAssetTypeActions_Base
+class
+// UE_DEPRECATED(5.2, "The AssetDefinition system is replacing AssetTypeActions and UAssetDefinition_SoundBase replaced FAssetTypeActions_SoundBase.  Please convert your FAssetTypeActions_SoundBase subclass to subclass from UAssetDefinition_SoundBase instead.")
+AUDIOEDITOR_API
+FAssetTypeActions_SoundBase : public FAssetTypeActions_Base
 {
 public:
 	// IAssetTypeActions Implementation
 	virtual FText GetName() const override { return NSLOCTEXT("AssetTypeActions", "AssetTypeActions_SoundBase", "Sound Base"); }
 	virtual FColor GetTypeColor() const override { return FColor(97, 85, 212); }
 	virtual UClass* GetSupportedClass() const override;
-	virtual void GetActions(const TArray<UObject*>& InObjects, struct FToolMenuSection& Section) override;
 	virtual bool AssetsActivatedOverride(const TArray<UObject*>& InObjects, EAssetTypeActivationMethod::Type ActivationType) override;
 	virtual uint32 GetCategories() override { return EAssetTypeCategories::Sounds; }
 	virtual bool CanFilter() override { return false; }
@@ -75,4 +76,3 @@ private:
 	/** Returns true if its possible to solo a sound */
 	bool CanExecuteSoloCommand(TArray<TWeakObjectPtr<USoundBase>> Objects) const;
 };
-PRAGMA_ENABLE_DEPRECATION_WARNINGS
diff --git a/Engine/Source/Editor/ContentBrowser/Public/ContentBrowserMenuContexts.h b/Engine/Source/Editor/ContentBrowser/Public/ContentBrowserMenuContexts.h
--- a/Engine/Source/Editor/ContentBrowser/Public/ContentBrowserMenuContexts.h
+++ b/Engine/Source/Editor/ContentBrowser/Public/ContentBrowserMenuContexts.h
@@ -70,6 +70,27 @@ public:
 		return LoadSelectedObjectsIfNeeded();
 	}
 
+	template<typename ExpectedAssetType>
+	TArray<ExpectedAssetType*> GetSelectedObjectsInMemory() const
+	{
+		TArray<ExpectedAssetType*> Result;
+		Result.Reserve(SelectedAssets.Num());
+		for (const FAssetData& Asset : SelectedAssets)
+		{
+			if (Asset.IsInstanceOf(ExpectedAssetType::StaticClass()))
+			{
+				if (UObject* AssetObject = Asset.FastGetAsset(false))
+				{
+					if (ExpectedAssetType* AssetObjectTyped = Cast<ExpectedAssetType>(AssetObject))
+					{
+						Result.Add(AssetObjectTyped);
+					}
+				}
+			}
+		}
+		return Result;
+	}
+
 	/**
 	 * Loads the selected assets (if needed) which is based on AssetViewUtils::LoadAssetsIfNeeded, this exists primarily
 	 * for backwards compatability.  Reliance on a black box to determine 'neededness' is not recommended, this function
@@ -115,13 +136,16 @@ public:
 		Result.Reserve(SelectedAssets.Num());
 		for (const FAssetData& Asset : SelectedAssets)
 		{
-			if (PredicateFilter(Asset))
+			if (Asset.IsInstanceOf(ExpectedAssetType::StaticClass()))
 			{
-				if (UObject* AssetObject = Asset.GetAsset(LoadTags))
+				if (PredicateFilter(Asset))
 				{
-					if (ExpectedAssetType* AssetObjectTyped = Cast<ExpectedAssetType>(AssetObject))
+					if (UObject* AssetObject = Asset.GetAsset(LoadTags))
 					{
-						Result.Add(AssetObjectTyped);
+						if (ExpectedAssetType* AssetObjectTyped = Cast<ExpectedAssetType>(AssetObject))
+						{
+							Result.Add(AssetObjectTyped);
+						}
 					}
 				}
 			}
@@ -129,6 +153,26 @@ public:
 		return Result;
 	}
 
+	/**
+	 * Loads the first selected valid asset and returns it.
+	 */
+	template<typename ExpectedAssetType>
+	ExpectedAssetType* LoadFirstSelectedObject(const TSet<FName>& LoadTags = {}) const
+	{   	
+		for (const FAssetData& Asset : SelectedAssets)
+		{
+			if (Asset.IsInstanceOf(ExpectedAssetType::StaticClass()))
+			{
+				if (ExpectedAssetType* ExpectedType = Cast<ExpectedAssetType>(Asset.GetAsset(LoadTags)))
+				{
+					return ExpectedType;
+				}
+			}	
+		}
+    
+		return nullptr;
+	}
+
 	/**
 	 * Returns a filtered array of assets that are of the desired class and potentially any subclasses.
 	 */
