Have audio Factories query whether SoundWave and Ambient Sound import is allowed.

diff --git a/Engine/Source/Editor/AudioEditor/Private/Factories/SoundFactory.cpp b/Engine/Source/Editor/AudioEditor/Private/Factories/SoundFactory.cpp
--- a/Engine/Source/Editor/AudioEditor/Private/Factories/SoundFactory.cpp
+++ b/Engine/Source/Editor/AudioEditor/Private/Factories/SoundFactory.cpp
@@ -23,10 +23,44 @@
 #include "EditorFramework/AssetImportData.h"
 #include "AudioCompressionSettingsUtils.h"
 #include "SoundFileIO/SoundFileIO.h"
+#include "Misc/NamePermissionList.h"
+#include "AssetToolsModule.h"
+#include "IAssetTools.h"
+
+// Disable user import
+static int32 EnableUserSoundwaveImportCvar = 1;
+FAutoConsoleVariableRef CVarEnableUserSoundwaveImport(
+	TEXT("au.EnableUserSoundwaveImport"),
+	EnableUserSoundwaveImportCvar,
+	TEXT("Enables letting the user import soundwaves in editor.\n")
+	TEXT("0: Disabled, 1: Enabled"),
+	ECVF_Default);
 
 
 namespace
 {
+
+	bool CanImportSoundWaves()
+	{
+		// disabled via cvar?
+		if(EnableUserSoundwaveImportCvar == 0)
+		{
+			return false;
+		}
+
+		IAssetTools& AssetTools = FAssetToolsModule::GetModule().Get();
+		TSharedPtr<FPathPermissionList> AssetClassPermissionList = AssetTools.GetAssetClassPathPermissionList(EAssetClassAction::ImportAsset);
+		if (AssetClassPermissionList && AssetClassPermissionList->HasFiltering())
+		{
+			if (!AssetClassPermissionList->PassesFilter(USoundWave::StaticClass()->GetPathName()))
+			{
+				return false;
+			}
+		}
+
+		return true;
+	}
+
 	void InsertSoundNode(USoundCue* SoundCue, UClass* NodeClass, int32 NodeIndex)
 	{
 		USoundNode* SoundNode = SoundCue->ConstructSoundNode<USoundNode>(NodeClass);
@@ -213,6 +247,14 @@ UObject* USoundFactory::CreateObject
 			}
 		}
 
+		if (!CanImportSoundWaves())
+		{
+			FMessageDialog::Open(EAppMsgType::Ok, FText::Format(NSLOCTEXT("SoundFactory", "Soundwave Import Not Allowed", "Soundwave import is not allowed ({0}: {1})"), FText::FromString(CuePackageName), Reason));
+			GEditor->GetEditorSubsystem<UImportSubsystem>()->BroadcastAssetPostImport(this, nullptr);
+			return nullptr;
+		}
+
+
 		// if we are creating the cue move it when necessary
 		UPackage* CuePackage = bMoveCue ? CreatePackage( *CuePackageName) : nullptr;
 
diff --git a/Engine/Source/Editor/UnrealEd/Classes/ActorFactories/ActorFactoryAmbientSound.h b/Engine/Source/Editor/UnrealEd/Classes/ActorFactories/ActorFactoryAmbientSound.h
--- a/Engine/Source/Editor/UnrealEd/Classes/ActorFactories/ActorFactoryAmbientSound.h
+++ b/Engine/Source/Editor/UnrealEd/Classes/ActorFactories/ActorFactoryAmbientSound.h
@@ -26,6 +26,8 @@ class UActorFactoryAmbientSound : public UActorFactory
 	virtual bool CanCreateActorFrom( const FAssetData& AssetData, FText& OutErrorMsg ) override;
 	virtual UObject* GetAssetFromActorInstance(AActor* ActorInstance) override;
 	//~ End UActorFactory Interface
+
+	static bool CanImportAmbientSounds();
 };
 
 
diff --git a/Engine/Source/Editor/UnrealEd/Private/Factories/ActorFactory.cpp b/Engine/Source/Editor/UnrealEd/Private/Factories/ActorFactory.cpp
--- a/Engine/Source/Editor/UnrealEd/Private/Factories/ActorFactory.cpp
+++ b/Engine/Source/Editor/UnrealEd/Private/Factories/ActorFactory.cpp
@@ -139,6 +139,9 @@ ActorFactory.cpp:
 
 #include "ClassViewerFilter.h"
 #include "ClassViewerModule.h"
+#include "Misc/NamePermissionList.h"
+#include "AssetToolsModule.h"
+#include "IAssetTools.h"
 
 DEFINE_LOG_CATEGORY(LogActorFactory);
 
@@ -1495,6 +1498,13 @@ UActorFactoryAmbientSound::UActorFactoryAmbientSound(const FObjectInitializer& O
 
 bool UActorFactoryAmbientSound::CanCreateActorFrom( const FAssetData& AssetData, FText& OutErrorMsg )
 {
+
+	if(!CanImportAmbientSounds())
+	{
+		OutErrorMsg = NSLOCTEXT("CanCreateActor", "AssetNotAllowed", "Ambient Sound Actors are disabled in this environment.");
+		return false;
+	}
+
 	//We allow creating AAmbientSounds without an existing sound asset
 	if ( UActorFactory::CanCreateActorFrom( AssetData, OutErrorMsg ) )
 	{
@@ -1546,6 +1556,21 @@ void UActorFactoryAmbientSound::PostCreateBlueprint( UObject* Asset, AActor* CDO
 	}
 }
 
+bool UActorFactoryAmbientSound::CanImportAmbientSounds()
+{
+	IAssetTools& AssetTools = FAssetToolsModule::GetModule().Get();
+	TSharedPtr<FPathPermissionList> AssetClassPermissionList = AssetTools.GetAssetClassPathPermissionList(EAssetClassAction::ImportAsset);
+	if (AssetClassPermissionList && AssetClassPermissionList->HasFiltering())
+	{
+		if (!AssetClassPermissionList->PassesFilter(AAmbientSound::StaticClass()->GetPathName()))
+		{
+			return false;
+		}
+	}
+
+	return true;
+}
+
 /*-----------------------------------------------------------------------------
 UActorFactoryClass
 -----------------------------------------------------------------------------*/
