Cvar-able soundwave duration import limit
#pf 63ffd8f4c00784bee26493b0

diff --git a/Engine/Source/Editor/AudioEditor/Private/Factories/SoundFactory.cpp b/Engine/Source/Editor/AudioEditor/Private/Factories/SoundFactory.cpp
--- a/Engine/Source/Editor/AudioEditor/Private/Factories/SoundFactory.cpp
+++ b/Engine/Source/Editor/AudioEditor/Private/Factories/SoundFactory.cpp
@@ -36,6 +36,14 @@ FAutoConsoleVariableRef CVarEnableUserSoundwaveImport(
 	TEXT("0: Disabled, 1: Enabled"),
 	ECVF_Default);
 
+static float SoundWaveImportLengthLimitInSecondsCVar = -1.f;
+FAutoConsoleVariableRef CVarSoundWaveImportLengthLimitInSeconds(
+	TEXT("au.SoundWaveImportLengthLimitInSeconds"),
+	SoundWaveImportLengthLimitInSecondsCVar,
+	TEXT("When set to a value > 0.0f, Soundwaves with durations greater than the value will fail to import.\n")
+	TEXT("if the value is < 0.0f, the length will be unlimited"),
+	ECVF_Default);
+
 
 namespace
 {
@@ -581,6 +589,17 @@ UObject* USoundFactory::CreateObject
 		Sound->NumChannels = ChannelCount;
 		Sound->TotalSamples = *WaveInfo.pSamplesPerSec * Sound->Duration;
 
+		const bool bLimitingSoundWaveLength = SoundWaveImportLengthLimitInSecondsCVar > 0.0f; 
+		if (bLimitingSoundWaveLength && Sound->Duration >= SoundWaveImportLengthLimitInSecondsCVar)
+		{
+			FMessageDialog::Open(EAppMsgType::Ok, FText::Format(NSLOCTEXT("SoundFactory", "Soundwave is too long to import"
+				, "{0} is {1} seconds in duration (this is over the limit of {2} seconds) {3}")
+				, FText::FromString(CuePackageName), Sound->Duration, SoundWaveImportLengthLimitInSecondsCVar, Reason));
+				
+			GEditor->GetEditorSubsystem<UImportSubsystem>()->BroadcastAssetPostImport(this, nullptr);
+			return nullptr;			
+		}
+
 		// Store the current file path and timestamp for re-import purposes
 		Sound->AssetImportData->Update(CurrentFilename);
 
