Reimporting an audio file over a Sound Wave marked as Ambisonic crashes editor

diff --git a/Engine/Source/Editor/AudioEditor/Private/Factories/SoundFactory.cpp b/Engine/Source/Editor/AudioEditor/Private/Factories/SoundFactory.cpp
--- a/Engine/Source/Editor/AudioEditor/Private/Factories/SoundFactory.cpp
+++ b/Engine/Source/Editor/AudioEditor/Private/Factories/SoundFactory.cpp
@@ -209,8 +209,8 @@ UObject* USoundFactory::FactoryCreateBinary
 
 	if (!SoundObject)
 	{
-		// Unrecognized sound format
-		Warn->Logf(ELogVerbosity::Error, TEXT("Unrecognized sound format '%s' in %s"), FileType, *Name.ToString());
+		// Inform user we failed to create the sound wave
+		Warn->Logf(ELogVerbosity::Error, TEXT("Failed to import sound wave %s"), *Name.ToString());
 		GEditor->GetEditorSubsystem<UImportSubsystem>()->BroadcastAssetPostImport(this, nullptr);
 	}
 
@@ -277,11 +277,6 @@ UObject* USoundFactory::CreateObject
 			// Will block internally on audio thread completing outstanding commands
 			AudioDeviceManager->StopSoundsUsingResource(ExistingSound, &ComponentsToRestart);
 
-			// We need to clear out any stale multichannel data on the sound wave in the case this is a reimport from multichannel to mono/stereo
-			ExistingSound->ChannelOffsets.Reset();
-			ExistingSound->ChannelSizes.Reset();
-			ExistingSound->bIsAmbisonics = false;
-
 			// Resource data is required to exist, if it hasn't been loaded yet,
 			// to properly flush compressed data.  This allows the new version
 			// to be auditioned in the editor properly.
@@ -699,6 +694,13 @@ UObject* USoundFactory::CreateObject
 
 		return Sound;
 	}
+	else
+	{ 
+		// Unrecognized sound format
+		Warn->Logf(ELogVerbosity::Error, TEXT("Unrecognized sound format '%s' in %s"), FileType, *Name.ToString());
+		GEditor->GetEditorSubsystem<UImportSubsystem>()->BroadcastAssetPostImport(this, nullptr);
+    }
+
 
 	return nullptr;
 }
