Add error message if we try and import a file bigger than int32_max

diff --git a/Engine/Source/Editor/AudioEditor/Private/Factories/SoundFactory.cpp b/Engine/Source/Editor/AudioEditor/Private/Factories/SoundFactory.cpp
--- a/Engine/Source/Editor/AudioEditor/Private/Factories/SoundFactory.cpp
+++ b/Engine/Source/Editor/AudioEditor/Private/Factories/SoundFactory.cpp
@@ -26,6 +26,7 @@
 #include "Misc/NamePermissionList.h"
 #include "AssetToolsModule.h"
 #include "IAssetTools.h"
+#include "Math/NumericLimits.h"
 
 // Disable user import
 static int32 EnableUserSoundwaveImportCvar = 1;
@@ -178,6 +179,16 @@ UObject* USoundFactory::FactoryCreateBinary
 {
 	GEditor->GetEditorSubsystem<UImportSubsystem>()->BroadcastAssetPreImport(this, Class, InParent, Name, FileType);
 
+	{ // Refuse to accept big files. We currently use TArray<> which will fail if we go over an int32.
+		const uint64 Size = BufferEnd - Buffer;
+		if (!IntFitsIn<int32>(Size))
+		{
+			Warn->Logf(ELogVerbosity::Error, TEXT("File '%s' is too big (%umb), Max=%umb"), *Name.ToString(), Size>>20, TNumericLimits<int32>::Max()>>20);
+			GEditor->GetEditorSubsystem<UImportSubsystem>()->BroadcastAssetPostImport(this, nullptr);
+			return nullptr;
+		}
+	}
+
 	UObject* SoundObject = nullptr;
 
 	// First, see if we support this file type in-engine:
