Fix for crash when overriding an Ambisonics sound asset with a non-Ambisonics file

diff --git a/Engine/Source/Editor/AudioEditor/Private/Factories/SoundFactory.cpp b/Engine/Source/Editor/AudioEditor/Private/Factories/SoundFactory.cpp
--- a/Engine/Source/Editor/AudioEditor/Private/Factories/SoundFactory.cpp
+++ b/Engine/Source/Editor/AudioEditor/Private/Factories/SoundFactory.cpp
@@ -220,9 +220,8 @@ UObject* USoundFactory::FactoryCreateBinary
 
 	if (!SoundObject)
 	{
-		// Unrecognized sound format
-		Warn->Logf(ELogVerbosity::Error, TEXT("Unrecognized sound format '%s' in %s"), FileType, *Name.ToString());
-		GEditor->GetEditorSubsystem<UImportSubsystem>()->BroadcastAssetPostImport(this, nullptr);
+		// Inform user we failed to create the sound wave
+		Warn->Logf(ELogVerbosity::Error, TEXT("Failed to import sound wave %s"), *Name.ToString());
 	}
 
 	return SoundObject;
@@ -288,11 +287,6 @@ UObject* USoundFactory::CreateObject
 			// Will block internally on audio thread completing outstanding commands
 			AudioDeviceManager->StopSoundsUsingResource(ExistingSound, &ComponentsToRestart);
 
-			// We need to clear out any stale multichannel data on the sound wave in the case this is a reimport from multichannel to mono/stereo
-			ExistingSound->ChannelOffsets.Reset();
-			ExistingSound->ChannelSizes.Reset();
-			ExistingSound->bIsAmbisonics = false;
-
 			// Resource data is required to exist, if it hasn't been loaded yet,
 			// to properly flush compressed data.  This allows the new version
 			// to be auditioned in the editor properly.
@@ -592,6 +586,15 @@ UObject* USoundFactory::CreateObject
 		}
 		else
 		{
+			// If this sound existed previously, we need to clear out any stale multichannel data on the sound wave in the case this is a reimport from multichannel to mono/stereo
+			if (ExistingSound)
+			{
+				ExistingSound->ChannelOffsets.Reset();
+				ExistingSound->ChannelSizes.Reset();
+				ExistingSound->bIsAmbisonics = false;
+			}
+			
+
 			// For mono and stereo assets, just copy the data into the buffer
 			FSharedBuffer UpdatedBuffer = FSharedBuffer::Clone(Buffer, RawWaveDataBufferSize);
 			Sound->RawData.UpdatePayload(UpdatedBuffer);
@@ -710,6 +713,12 @@ UObject* USoundFactory::CreateObject
 
 		return Sound;
 	}
+	else
+	{
+		// Unrecognized sound format
+		Warn->Logf(ELogVerbosity::Error, TEXT("Unrecognized sound format '%s' in %s"), FileType, *Name.ToString());
+		GEditor->GetEditorSubsystem<UImportSubsystem>()->BroadcastAssetPostImport(this, nullptr);
+	}
 
 	return nullptr;
 }
