Copy Sound Wave compression type and value when importing a file using a template soundwave

diff --git a/Engine/Source/Editor/AudioEditor/Private/Factories/SoundFactory.cpp b/Engine/Source/Editor/AudioEditor/Private/Factories/SoundFactory.cpp
--- a/Engine/Source/Editor/AudioEditor/Private/Factories/SoundFactory.cpp
+++ b/Engine/Source/Editor/AudioEditor/Private/Factories/SoundFactory.cpp
@@ -431,6 +431,8 @@ UObject* USoundFactory::CreateObject
 		{
 			Sound->SoundClassObject = TemplateSoundWave->SoundClassObject;
 			Sound->ConcurrencySet = TemplateSoundWave->ConcurrencySet;
+			Sound->CompressionQuality = TemplateSoundWave->CompressionQuality;
+			Sound->SoundAssetCompressionType = TemplateSoundWave->SoundAssetCompressionType;
 
 			// we do not want to inherit these values from the template, as the data may be incorrect
 			// rather we re-parse them from the incoming file.
diff --git a/Engine/Source/Runtime/Engine/Classes/Sound/SoundWave.h b/Engine/Source/Runtime/Engine/Classes/Sound/SoundWave.h
--- a/Engine/Source/Runtime/Engine/Classes/Sound/SoundWave.h
+++ b/Engine/Source/Runtime/Engine/Classes/Sound/SoundWave.h
@@ -1257,6 +1257,7 @@ public:
 
 private:
 	friend class FSoundWaveProxy;
+	friend class USoundFactory;
 };
 
 
