Fix pinned audio asset types editor preferences not being observed

diff --git a/Engine/Source/Editor/AudioEditor/Classes/AudioEditorSettings.h b/Engine/Source/Editor/AudioEditor/Classes/AudioEditorSettings.h
--- a/Engine/Source/Editor/AudioEditor/Classes/AudioEditorSettings.h
+++ b/Engine/Source/Editor/AudioEditor/Classes/AudioEditorSettings.h
@@ -19,19 +19,19 @@ class AUDIOEDITOR_API UAudioEditorSettings : public UDeveloperSettings
 	GENERATED_BODY()
 
 public:
-	/** Whether to pin the Sound Cue asset type when creating new assets. */
+	/** Whether to pin the Sound Cue asset type when creating new assets. Requires editor restart to take effect. */
 	UPROPERTY(EditAnywhere, config, Category = AssetMenu)
 	bool bPinSoundCueInAssetMenu = true;
 
-	/** Whether to pin the Sound Cue Template asset type when creating new assets. */
+	/** Whether to pin the Sound Cue Template asset type when creating new assets. Requires editor restart to take effect. */
 	UPROPERTY(EditAnywhere, config, Category = AssetMenu)
 	bool bPinSoundCueTemplateInAssetMenu = false;
 
-	/** Whether to pin the Sound Attenuation asset type when creating new assets. */
+	/** Whether to pin the Sound Attenuation asset type when creating new assets. Requires editor restart to take effect. */
 	UPROPERTY(EditAnywhere, config, Category = AssetMenu)
 	bool bPinSoundAttenuationInAssetMenu = true;
 
-	/** Whether to pin the Sound Concurrency asset type when creating new assets. */
+	/** Whether to pin the Sound Concurrency asset type when creating new assets. Requires editor restart to take effect. */
 	UPROPERTY(EditAnywhere, config, Category = AssetMenu)
 	bool bPinSoundConcurrencyInAssetMenu = true;
 
diff --git a/Engine/Source/Editor/AudioEditor/Private/AssetTypeActions/AssetDefinition_SoundAttenuation.h b/Engine/Source/Editor/AudioEditor/Private/AssetTypeActions/AssetDefinition_SoundAttenuation.h
--- a/Engine/Source/Editor/AudioEditor/Private/AssetTypeActions/AssetDefinition_SoundAttenuation.h
+++ b/Engine/Source/Editor/AudioEditor/Private/AssetTypeActions/AssetDefinition_SoundAttenuation.h
@@ -5,6 +5,7 @@
 #include "CoreMinimal.h"
 #include "Sound/SoundAttenuation.h"
 #include "AssetDefinitionDefault.h"
+#include "AudioEditorSettings.h"
 
 #include "AssetDefinition_SoundAttenuation.generated.h"
 
@@ -20,7 +21,13 @@ public:
 	virtual TSoftClassPtr<UObject> GetAssetClass() const override { return USoundAttenuation::StaticClass(); }
 	virtual TConstArrayView<FAssetCategoryPath> GetAssetCategories() const override
 	{
+		static const auto PinnedCategories = { EAssetCategoryPaths::Audio };
 		static const auto Categories = { FAssetCategoryPath(EAssetCategoryPaths::Audio, NSLOCTEXT("AssetTypeActions", "AssetSoundAttenuationSubMenu", "Spatialization")) };
+		
+		if (GetDefault<UAudioEditorSettings>()->bPinSoundAttenuationInAssetMenu)
+		{
+			return PinnedCategories;
+		}
 		return Categories;
 	}
 	// UAssetDefinition End
diff --git a/Engine/Source/Editor/AudioEditor/Private/AssetTypeActions/AssetDefinition_SoundConcurrency.h b/Engine/Source/Editor/AudioEditor/Private/AssetTypeActions/AssetDefinition_SoundConcurrency.h
--- a/Engine/Source/Editor/AudioEditor/Private/AssetTypeActions/AssetDefinition_SoundConcurrency.h
+++ b/Engine/Source/Editor/AudioEditor/Private/AssetTypeActions/AssetDefinition_SoundConcurrency.h
@@ -5,6 +5,7 @@
 #include "CoreMinimal.h"
 #include "Sound/SoundConcurrency.h"
 #include "AssetDefinitionDefault.h"
+#include "AudioEditorSettings.h"
 
 #include "AssetDefinition_SoundConcurrency.generated.h"
 
@@ -20,7 +21,13 @@ public:
 	virtual TSoftClassPtr<UObject> GetAssetClass() const override { return USoundConcurrency::StaticClass(); }
 	virtual TConstArrayView<FAssetCategoryPath> GetAssetCategories() const override
 	{
+		static const auto PinnedCategories = { EAssetCategoryPaths::Audio };
 		static const auto Categories = { FAssetCategoryPath(EAssetCategoryPaths::Audio, NSLOCTEXT("AssetTypeActions", "AssetSoundConcurrencySubMenu", "Mix")) };
+
+		if (GetDefault<UAudioEditorSettings>()->bPinSoundConcurrencyInAssetMenu)
+		{
+			return PinnedCategories;
+		}
 		return Categories;
 	}
 	// UAssetDefinition End
