Clear out a reference count so that users can MoveToUnique without a copy.

diff --git a/Engine/Source/Editor/AudioEditor/Private/Factories/SoundFactory.cpp b/Engine/Source/Editor/AudioEditor/Private/Factories/SoundFactory.cpp
--- a/Engine/Source/Editor/AudioEditor/Private/Factories/SoundFactory.cpp
+++ b/Engine/Source/Editor/AudioEditor/Private/Factories/SoundFactory.cpp
@@ -598,8 +598,8 @@ UObject* USoundFactory::CreateObject
 			
 
 			// For mono and stereo assets, just copy the data into the buffer
-			FSharedBuffer UpdatedBuffer = FSharedBuffer::Clone(Buffer, RawWaveDataBufferSize);
-			Sound->RawData.UpdatePayload(UpdatedBuffer);
+			// Clone directly as a param so that if anyone MoveToUniques it then its a steal not a copy.
+			Sound->RawData.UpdatePayload(FSharedBuffer::Clone(Buffer, RawWaveDataBufferSize));
 
 		}
 
