Warn on importing sound waves that have non zero ends.

diff --git a/Engine/Source/Editor/AudioEditor/Private/Factories/SoundFactory.cpp b/Engine/Source/Editor/AudioEditor/Private/Factories/SoundFactory.cpp
--- a/Engine/Source/Editor/AudioEditor/Private/Factories/SoundFactory.cpp
+++ b/Engine/Source/Editor/AudioEditor/Private/Factories/SoundFactory.cpp
@@ -449,6 +449,30 @@ UObject* USoundFactory::CreateObject
 		int32 NumSamples = WaveInfo.SampleDataSize / SizeOfSample;
 		int32 NumFrames = NumSamples / ChannelCount;
 
+		// Starting or ending on a non-zero sample creates a pop unless you have very specific use cases
+		// (doesn't apply to ambi)
+		if (!bIsFuMa && !bIsAmbiX)
+		{
+			for (int32 Chan = 0; Chan < ChannelCount; ++Chan)
+			{
+				const int16* FirstFrame = (const int16*)WaveInfo.SampleDataStart;
+				const int16* LastFrame = FirstFrame + (NumFrames-1)*ChannelCount;
+
+				FirstFrame += Chan;
+				LastFrame += Chan;
+
+				if (FirstFrame[0] != 0)
+				{
+					Warn->Logf(ELogVerbosity::Warning, TEXT("Channel %d starts with a non zero value (%d) - will likely pop and not loop well: '%s'"), Chan, *FirstFrame, *Name.ToString());
+				}
+
+				if (LastFrame[0] != 0)
+				{
+					Warn->Logf(ELogVerbosity::Warning, TEXT("Channel %d ends with a non zero value (%d) - will likely pop and not loop well: '%s'"), Chan, *LastFrame, *Name.ToString());
+				}
+			}
+		}
+
 		if (ChannelCount > 2)
 		{
 			// We need to deinterleave the raw PCM data in the multi-channel file reuse a scratch buffer
