Make the DC offset audio import warning less spammy and annoying.

diff --git a/Engine/Source/Editor/AudioEditor/Private/Factories/SoundFactory.cpp b/Engine/Source/Editor/AudioEditor/Private/Factories/SoundFactory.cpp
--- a/Engine/Source/Editor/AudioEditor/Private/Factories/SoundFactory.cpp
+++ b/Engine/Source/Editor/AudioEditor/Private/Factories/SoundFactory.cpp
@@ -453,6 +453,9 @@ UObject* USoundFactory::CreateObject
 		// (doesn't apply to ambi)
 		if (!bIsFuMa && !bIsAmbiX)
 		{
+			int32 EndChannelNonZero = 0;
+			int32 StartChannelNonZero = 0;
+			constexpr int16 LogDCLimit = 100; // amount of non zero dc we allow before we nag
 			for (int32 Chan = 0; Chan < ChannelCount; ++Chan)
 			{
 				const int16* FirstFrame = (const int16*)WaveInfo.SampleDataStart;
@@ -461,16 +464,29 @@ UObject* USoundFactory::CreateObject
 				FirstFrame += Chan;
 				LastFrame += Chan;
 
-				if (FirstFrame[0] != 0)
+				if (FMath::Abs(FirstFrame[0]) > LogDCLimit)
 				{
-					Warn->Logf(ELogVerbosity::Warning, TEXT("Channel %d starts with a non zero value (%d) - will likely pop and not loop well: '%s'"), Chan, *FirstFrame, *Name.ToString());
+					StartChannelNonZero++;
 				}
 
-				if (LastFrame[0] != 0)
+				if (FMath::Abs(LastFrame[0]) > LogDCLimit)
 				{
-					Warn->Logf(ELogVerbosity::Warning, TEXT("Channel %d ends with a non zero value (%d) - will likely pop and not loop well: '%s'"), Chan, *LastFrame, *Name.ToString());
+					EndChannelNonZero++;
 				}
 			}
+
+			if (EndChannelNonZero && StartChannelNonZero)
+			{
+				Warn->Logf(ELogVerbosity::Log, TEXT("Imported audio has DC offsets larger than %d at the end of %d channel%s and the start of %d channel%s - may pop and not loop well: '%s'"), LogDCLimit, EndChannelNonZero, EndChannelNonZero > 1 ? TEXT("s") : TEXT(""), StartChannelNonZero, StartChannelNonZero > 1 ? TEXT("s") : TEXT(""), *Name.ToString());
+			}
+			else if (EndChannelNonZero)
+			{
+				Warn->Logf(ELogVerbosity::Log, TEXT("Imported audio has DC offsets larger than %d at the end of %d channel%s - may pop and not loop well: '%s'"), LogDCLimit, EndChannelNonZero, EndChannelNonZero > 1 ? TEXT("s") : TEXT(""), *Name.ToString());
+			}
+			else if (StartChannelNonZero)
+			{
+				Warn->Logf(ELogVerbosity::Log, TEXT("Imported audio has DC offsets larger than %d at the start of %d channel%s - may pop and not loop well: '%s'"), LogDCLimit, StartChannelNonZero, StartChannelNonZero > 1 ? TEXT("s") : TEXT(""), *Name.ToString());
+			}
 		}
 
 		if (ChannelCount > 2)
