Fixed crash when opening some audio assets in the editor, caused by a direct call to the constructor of the detail customizations.
Made those classes' constructors take a protected/private token so they are forced to use a MakeInstance() call.

diff --git a/Engine/Source/Editor/AudioEditor/Private/SubmixDetailsCustomization.cpp b/Engine/Source/Editor/AudioEditor/Private/SubmixDetailsCustomization.cpp
--- a/Engine/Source/Editor/AudioEditor/Private/SubmixDetailsCustomization.cpp
+++ b/Engine/Source/Editor/AudioEditor/Private/SubmixDetailsCustomization.cpp
@@ -104,16 +104,24 @@ TSharedRef<SWidget> FNameSelectorGenerator::HandleResponseComboBoxGenerateWidget
 		.Font(IDetailLayoutBuilder::GetDetailFont());
 }
 
+FNameSelectorGenerator::FNameSelectorGenerator(FProtectedToken)
+{
+}
+
 TSharedRef<FNameSelectorGenerator> FNameSelectorGenerator::MakeInstance()
 {
-	TSharedRef<FNameSelectorGenerator> Result = MakeShared<FNameSelectorGenerator>();
+	TSharedRef<FNameSelectorGenerator> Result = MakeShared<FNameSelectorGenerator>(FProtectedToken{});
 	Result->SetWeakThis(Result.ToWeakPtr());
 	return Result;
 }
 
+FSoundfieldSubmixDetailsCustomization::FSoundfieldSubmixDetailsCustomization(FPrivateToken)
+{
+}
+
 TSharedRef<IDetailCustomization> FSoundfieldSubmixDetailsCustomization::MakeInstance()
 {
-	return MakeShareable(new FSoundfieldSubmixDetailsCustomization);
+	return MakeShared<FSoundfieldSubmixDetailsCustomization>(FPrivateToken{});
 }
 
 void FSoundfieldSubmixDetailsCustomization::CustomizeDetails(IDetailLayoutBuilder& DetailLayout)
@@ -176,9 +184,14 @@ void FSoundfieldSubmixDetailsCustomization::CustomizeDetails(IDetailLayoutBuilde
 		];
 }
 
+FEndpointSubmixDetailsCustomization::FEndpointSubmixDetailsCustomization(FProtectedToken ProtectedToken)
+	: FNameSelectorGenerator(ProtectedToken)
+{
+}
+
 TSharedRef<IDetailCustomization> FEndpointSubmixDetailsCustomization::MakeInstance()
 {
-	TSharedRef<FEndpointSubmixDetailsCustomization> Result = MakeShared<FEndpointSubmixDetailsCustomization>();
+	TSharedRef<FEndpointSubmixDetailsCustomization> Result = MakeShared<FEndpointSubmixDetailsCustomization>(FProtectedToken{});
 	Result->SetWeakThis(StaticCastWeakPtr<FNameSelectorGenerator>(Result.ToWeakPtr()));
 	return Result;
 }
@@ -187,7 +200,7 @@ void FEndpointSubmixDetailsCustomization::CustomizeDetails(IDetailLayoutBuilder&
 {
 	if (!EndpointTypeNameSelectorGenerator)
 	{
-		EndpointTypeNameSelectorGenerator = MakeShareable(new FNameSelectorGenerator());
+		EndpointTypeNameSelectorGenerator = FNameSelectorGenerator::MakeInstance();
 	}
 
 	IDetailCategoryBuilder& EndpointCategory = DetailLayout.EditCategory(TEXT("Endpoint"));
@@ -236,9 +249,14 @@ void FEndpointSubmixDetailsCustomization::CustomizeDetails(IDetailLayoutBuilder&
 		];
 }
 
+FSoundfieldEndpointSubmixDetailsCustomization::FSoundfieldEndpointSubmixDetailsCustomization(FProtectedToken ProtectedToken)
+	: FNameSelectorGenerator(ProtectedToken)
+{
+}
+
 TSharedRef<IDetailCustomization> FSoundfieldEndpointSubmixDetailsCustomization::MakeInstance()
 {
-	TSharedRef<FSoundfieldEndpointSubmixDetailsCustomization> Result = MakeShared<FSoundfieldEndpointSubmixDetailsCustomization>();
+	TSharedRef<FSoundfieldEndpointSubmixDetailsCustomization> Result = MakeShared<FSoundfieldEndpointSubmixDetailsCustomization>(FProtectedToken{});
 	Result->SetWeakThis(StaticCastWeakPtr<FNameSelectorGenerator>(Result.ToWeakPtr()));
 	return Result;
 }
@@ -247,7 +265,7 @@ void FSoundfieldEndpointSubmixDetailsCustomization::CustomizeDetails(IDetailLayo
 {
 	if (!EndpointTypeNameSelectorGenerator)
 	{
-		EndpointTypeNameSelectorGenerator = MakeShareable(new FNameSelectorGenerator());
+		EndpointTypeNameSelectorGenerator = FNameSelectorGenerator::MakeInstance();
 	}
 
 	IDetailCategoryBuilder& EndpointCategory = DetailLayout.EditCategory(TEXT("Endpoint"));
diff --git a/Engine/Source/Editor/AudioEditor/Public/SubmixDetailsCustomization.h b/Engine/Source/Editor/AudioEditor/Public/SubmixDetailsCustomization.h
--- a/Engine/Source/Editor/AudioEditor/Public/SubmixDetailsCustomization.h
+++ b/Engine/Source/Editor/AudioEditor/Public/SubmixDetailsCustomization.h
@@ -23,8 +23,12 @@ class SWidget;
 // Utility class to build combo boxes out of arrays of names.
 class FNameSelectorGenerator
 {
+protected:
+	struct FProtectedToken { explicit FProtectedToken() = default; };
 
 public:
+	FNameSelectorGenerator(FProtectedToken);
+
 	struct FNameSelectorCallbacks
 	{
 		TUniqueFunction<void(FName)> OnNewNameSelected;
@@ -56,7 +60,12 @@ private:
 
 class AUDIOEDITOR_API FSoundfieldSubmixDetailsCustomization : public IDetailCustomization
 {
+protected:
+	struct FPrivateToken { explicit FPrivateToken() = default; };
+
 public:
+	FSoundfieldSubmixDetailsCustomization(FPrivateToken);
+
 	// Makes a new instance of this detail layout class
 	static TSharedRef<IDetailCustomization> MakeInstance();
 
@@ -71,6 +80,8 @@ private:
 class AUDIOEDITOR_API FEndpointSubmixDetailsCustomization : public IDetailCustomization, public FNameSelectorGenerator
 {
 public:
+	FEndpointSubmixDetailsCustomization(FProtectedToken);
+
 	// Makes a new instance of this detail layout class
 	static TSharedRef<IDetailCustomization> MakeInstance();
 
@@ -85,6 +96,8 @@ private:
 class AUDIOEDITOR_API FSoundfieldEndpointSubmixDetailsCustomization : public IDetailCustomization, public FNameSelectorGenerator
 {
 public:
+	FSoundfieldEndpointSubmixDetailsCustomization(FProtectedToken);
+
 	// Makes a new instance of this detail layout class
 	static TSharedRef<IDetailCustomization> MakeInstance();
 
