Playing the MetaSound from the Content Browser does not reflect the changes on the asset unless it was previously saved or played in the MetaSound editor
#rnx

diff --git a/Engine/Source/Editor/AudioEditor/Private/AssetTypeActions/AssetDefinition_SoundBase.cpp b/Engine/Source/Editor/AudioEditor/Private/AssetTypeActions/AssetDefinition_SoundBase.cpp
--- a/Engine/Source/Editor/AudioEditor/Private/AssetTypeActions/AssetDefinition_SoundBase.cpp
+++ b/Engine/Source/Editor/AudioEditor/Private/AssetTypeActions/AssetDefinition_SoundBase.cpp
@@ -91,7 +91,7 @@ EAssetCommandResult UAssetDefinition_SoundBase::ActivateSoundBase(const FAssetAc
 	return EAssetCommandResult::Unhandled;
 }
 
-TSharedPtr<SWidget> UAssetDefinition_SoundBase::GetSoundBaseThumbnailOverlay(const FAssetData& InAssetData, TUniqueFunction<FReply()>&& OnClickedLambdaOverride)
+TSharedPtr<SWidget> UAssetDefinition_SoundBase::GetSoundBaseThumbnailOverlay(const FAssetData& InAssetData, TFunction<FReply()>&& OnClicked)
 {
 	auto OnGetDisplayBrushLambda = [InAssetData]() -> const FSlateBrush*
 	{
@@ -103,19 +103,23 @@ TSharedPtr<SWidget> UAssetDefinition_SoundBase::GetSoundBaseThumbnailOverlay(con
 		return FAppStyle::GetBrush("MediaAsset.AssetActions.Play.Large");
 	};
 
-	auto OnClickedLambda = [InAssetData]() -> FReply
+	auto OnClickedLambda = MoveTemp(OnClicked);
+	if (!OnClickedLambda)
 	{
-		if (UE::AudioEditor::IsSoundPlaying(InAssetData))
+		OnClickedLambda = [InAssetData]() -> FReply
 		{
-			UE::AudioEditor::StopSound();
-		}
-		else
-		{
-			// Load and play sound
-			UE::AudioEditor::PlaySound(Cast<USoundBase>(InAssetData.GetAsset()));
-		}
-		return FReply::Handled();
-	};
+			if (UE::AudioEditor::IsSoundPlaying(InAssetData))
+			{
+				UE::AudioEditor::StopSound();
+			}
+			else
+			{
+				// Load and play sound
+				UE::AudioEditor::PlaySound(Cast<USoundBase>(InAssetData.GetAsset()));
+			}
+			return FReply::Handled();
+		};
+	}
 
 	auto OnToolTipTextLambda = [InAssetData]() -> FText
 	{
diff --git a/Engine/Source/Editor/AudioEditor/Public/AssetTypeActions/AssetDefinition_SoundBase.h b/Engine/Source/Editor/AudioEditor/Public/AssetTypeActions/AssetDefinition_SoundBase.h
--- a/Engine/Source/Editor/AudioEditor/Public/AssetTypeActions/AssetDefinition_SoundBase.h
+++ b/Engine/Source/Editor/AudioEditor/Public/AssetTypeActions/AssetDefinition_SoundBase.h
@@ -50,6 +50,6 @@ public:
 	static bool CanExecuteSoloCommand(const FToolMenuContext& InContext);
 
 	// Asset definition static utilities
-	static TSharedPtr<SWidget> GetSoundBaseThumbnailOverlay(const FAssetData& InAssetData, TUniqueFunction<FReply()>&& OnClickedLambda);
+	static TSharedPtr<SWidget> GetSoundBaseThumbnailOverlay(const FAssetData& InAssetData, TFunction<FReply()>&& OnClicked);
 	static EAssetCommandResult ActivateSoundBase(const FAssetActivateArgs& ActivateArgs);
 };
