[AudioEditor] Sound Waves renamed to Sound Base, no longer filter in content browser

diff --git a/Engine/Source/Editor/AudioEditor/Private/AudioEditorModule.cpp b/Engine/Source/Editor/AudioEditor/Private/AudioEditorModule.cpp
--- a/Engine/Source/Editor/AudioEditor/Private/AudioEditorModule.cpp
+++ b/Engine/Source/Editor/AudioEditor/Private/AudioEditorModule.cpp
@@ -179,6 +179,7 @@ public:
 
 		AssetTools.RegisterAssetTypeActions(MakeShared<FAssetTypeActions_DialogueWave>());
 		AssetTools.RegisterAssetTypeActions(MakeShared<FAssetTypeActions_SoundClass>());
+		AssetTools.RegisterAssetTypeActions(MakeShared<FAssetTypeActions_SoundWave>());
 	}
 
 	virtual void RegisterAudioMixerAssetActions() override
