Audition ActionButtonExtensions for MetaSoundSource_AssetDefinition
#rnx

diff --git a/Engine/Plugins/Runtime/Metasound/Source/MetasoundEditor/Private/MetasoundAssetDefinitions.cpp b/Engine/Plugins/Runtime/Metasound/Source/MetasoundEditor/Private/MetasoundAssetDefinitions.cpp
--- a/Engine/Plugins/Runtime/Metasound/Source/MetasoundEditor/Private/MetasoundAssetDefinitions.cpp
+++ b/Engine/Plugins/Runtime/Metasound/Source/MetasoundEditor/Private/MetasoundAssetDefinitions.cpp
@@ -396,6 +396,11 @@ EAssetCommandResult UAssetDefinition_MetaSoundSource::ActivateAssets(const FAsse
 	return EAssetCommandResult::Unhandled;
 }
 
+void UAssetDefinition_MetaSoundSource::GetAssetActionButtonExtensions(const FAssetData& InAssetData, TArray<FAssetButtonActionExtension>& OutExtensions) const
+{
+	UAssetDefinition_SoundBase::GetSoundBaseAssetActionButtonExtensions(InAssetData, OutExtensions);
+}
+
 namespace MenuExtension_MetaSoundSourceTemplate
 {
 	template <typename TClass>
diff --git a/Engine/Plugins/Runtime/Metasound/Source/MetasoundEditor/Private/MetasoundAssetDefinitions.h b/Engine/Plugins/Runtime/Metasound/Source/MetasoundEditor/Private/MetasoundAssetDefinitions.h
--- a/Engine/Plugins/Runtime/Metasound/Source/MetasoundEditor/Private/MetasoundAssetDefinitions.h
+++ b/Engine/Plugins/Runtime/Metasound/Source/MetasoundEditor/Private/MetasoundAssetDefinitions.h
@@ -40,6 +40,7 @@ public:
 	virtual TSharedPtr<SWidget> GetThumbnailOverlay(const FAssetData& InAssetData) const override;
 	virtual bool GetThumbnailActionOverlay(const FAssetData& InAssetData, FAssetActionThumbnailOverlayInfo& OutActionOverlayInfo) const override;
 	virtual EAssetCommandResult ActivateAssets(const FAssetActivateArgs& ActivateArgs) const override;
+	virtual void GetAssetActionButtonExtensions(const FAssetData& InAssetData, TArray<FAssetButtonActionExtension>& OutExtensions) const override;
 
  	virtual const FSlateBrush* GetThumbnailBrush(const FAssetData& InAssetData, const FName InClassName) const override;
  	virtual const FSlateBrush* GetIconBrush(const FAssetData& InAssetData, const FName InClassName) const override;
diff --git a/Engine/Source/Editor/AudioEditor/Private/AssetTypeActions/AssetDefinition_SoundBase.cpp b/Engine/Source/Editor/AudioEditor/Private/AssetTypeActions/AssetDefinition_SoundBase.cpp
--- a/Engine/Source/Editor/AudioEditor/Private/AssetTypeActions/AssetDefinition_SoundBase.cpp
+++ b/Engine/Source/Editor/AudioEditor/Private/AssetTypeActions/AssetDefinition_SoundBase.cpp
@@ -179,42 +179,7 @@ EAssetCommandResult UAssetDefinition_SoundBase::ActivateAssets(const FAssetActiv
 
 void UAssetDefinition_SoundBase::GetAssetActionButtonExtensions(const FAssetData& InAssetData, TArray<FAssetButtonActionExtension>& OutExtensions) const
 {
-	FAssetButtonActionExtension AssetButtonActionExtension
-	{	
-		.PickTooltipAttribute = TAttribute<FText>::CreateLambda([InAssetData]() -> const FText
-		{
-			if (UE::AudioEditor::IsSoundPlaying(InAssetData))
-			{
-				return LOCTEXT("SoundAudition_PlaySoundToolTip", "Stop selected sound");
-			}
-
-			return LOCTEXT("SoundAudition_StopSoundToolTip", "Play selected sound");
-		}),
-		.PickBrushAttribute = TAttribute<const FSlateBrush*>::CreateLambda([InAssetData]() -> const FSlateBrush* {
-		
-			if (UE::AudioEditor::IsSoundPlaying(InAssetData))
-			{
-				return FAppStyle::GetBrush("MediaAsset.AssetActions.Stop.Small");
-			}
-
-			return FAppStyle::GetBrush("MediaAsset.AssetActions.Play.Small");
-		}),
-		.OnClicked = FOnClicked::CreateLambda([InAssetData]() -> FReply 
-			{
-				if (UE::AudioEditor::IsSoundPlaying(InAssetData))
-				{
-					UE::AudioEditor::StopSound();
-				}
-				else
-				{
-					// Load and play sound
-					UE::AudioEditor::PlaySound(Cast<USoundBase>(InAssetData.GetAsset()));
-				}
-			return FReply::Handled(); 
-		})
-	};
-	
-	OutExtensions.Add(AssetButtonActionExtension);
+	UAssetDefinition_SoundBase::GetSoundBaseAssetActionButtonExtensions(InAssetData, OutExtensions);
 }
 
 TSharedPtr<SWidget> UAssetDefinition_SoundBase::GetThumbnailOverlay(const FAssetData& InAssetData) const
@@ -498,6 +463,47 @@ bool UAssetDefinition_SoundBase::GetThumbnailActionOverlay(const FAssetData& InA
 		return false;
 	}
 
+	void UAssetDefinition_SoundBase::GetSoundBaseAssetActionButtonExtensions(const FAssetData& InAssetData, TArray<FAssetButtonActionExtension>& OutExtensions)
+	{
+		FAssetButtonActionExtension AssetButtonActionExtension
+		{
+			.PickTooltipAttribute = TAttribute<FText>::CreateLambda([InAssetData]() -> const FText
+			{
+				if (UE::AudioEditor::IsSoundPlaying(InAssetData))
+				{
+					return LOCTEXT("SoundAudition_PlaySoundToolTip", "Stop selected sound");
+				}
+
+				return LOCTEXT("SoundAudition_StopSoundToolTip", "Play selected sound");
+			}),
+			.PickBrushAttribute = TAttribute<const FSlateBrush*>::CreateLambda([InAssetData]() -> const FSlateBrush* {
+
+				if (UE::AudioEditor::IsSoundPlaying(InAssetData))
+				{
+					return FAppStyle::GetBrush("MediaAsset.AssetActions.Stop.Small");
+				}
+
+				return FAppStyle::GetBrush("MediaAsset.AssetActions.Play.Small");
+			}),
+			.OnClicked = FOnClicked::CreateLambda([InAssetData]() -> FReply
+				{
+					if (UE::AudioEditor::IsSoundPlaying(InAssetData))
+					{
+						UE::AudioEditor::StopSound();
+					}
+					else
+					{
+						// Load and play sound
+						UE::AudioEditor::PlaySound(Cast<USoundBase>(InAssetData.GetAsset()));
+					}
+
+				return FReply::Handled();
+			})
+		};
+
+		OutExtensions.Add(AssetButtonActionExtension);
+	}
+
 namespace MenuExtension_SoundBase
 {
 	static void RegisterAssetActions(UClass* InClass)
diff --git a/Engine/Source/Editor/AudioEditor/Public/AssetTypeActions/AssetDefinition_SoundBase.h b/Engine/Source/Editor/AudioEditor/Public/AssetTypeActions/AssetDefinition_SoundBase.h
--- a/Engine/Source/Editor/AudioEditor/Public/AssetTypeActions/AssetDefinition_SoundBase.h
+++ b/Engine/Source/Editor/AudioEditor/Public/AssetTypeActions/AssetDefinition_SoundBase.h
@@ -53,5 +53,6 @@ public:
 
 	// Asset definition static utilities
 	static TSharedPtr<SWidget> GetSoundBaseThumbnailOverlay(const FAssetData& InAssetData, TFunction<FReply()>&& OnClicked);
+	static void GetSoundBaseAssetActionButtonExtensions(const FAssetData& InAssetData, TArray<FAssetButtonActionExtension>& OutExtensions);
 	static EAssetCommandResult ActivateSoundBase(const FAssetActivateArgs& ActivateArgs);
 };
