Disable Create Blueprint asset action for audio assets in restricted mode
- Move audio asset component broker to audio editor module

diff --git a/Engine/Source/Editor/AudioEditor/Private/AudioEditorModule.cpp b/Engine/Source/Editor/AudioEditor/Private/AudioEditorModule.cpp
--- a/Engine/Source/Editor/AudioEditor/Private/AudioEditorModule.cpp
+++ b/Engine/Source/Editor/AudioEditor/Private/AudioEditorModule.cpp
@@ -8,6 +8,7 @@
 #include "AssetTypeActions/AssetTypeActions_SoundEffectPreset.h"
 #include "AssetTypeActions/AssetTypeActions_SoundSubmix.h"
 #include "ClassTemplateEditorSubsystem.h"
+#include "ComponentAssetBroker.h"
 #include "Components/SynthComponent.h"
 #include "EdGraphUtilities.h"
 #include "Factories/ReimportSoundFactory.h"
@@ -66,6 +67,36 @@ static const FVector2D Icon64 = FVector2D(64.0f, 64.0f);
 
 #define SET_AUDIO_ICON_SVG_SIMPLE(CLASS_NAME) SET_AUDIO_ICON_SVG(CLASS_NAME, CLASS_NAME)
 
+UClass* FAudioComponentBroker::GetSupportedAssetClass()
+{
+	return USoundBase::StaticClass();
+}
+
+bool FAudioComponentBroker::AssignAssetToComponent(UActorComponent* InComponent, UObject* InAsset)
+{
+	if (UAudioComponent* AudioComp = Cast<UAudioComponent>(InComponent))
+	{
+		USoundBase* Sound = Cast<USoundBase>(InAsset);
+
+		if ((Sound != NULL) || (InAsset == NULL))
+		{
+			AudioComp->SetSound(Sound);
+			return true;
+		}
+	}
+
+	return false;
+}
+
+UObject* FAudioComponentBroker::GetAssetFromComponent(UActorComponent* InComponent)
+{
+	if (UAudioComponent* AudioComp = Cast<UAudioComponent>(InComponent))
+	{
+		return AudioComp->Sound;
+	}
+	return NULL;
+}
+
 class FAudioEditorModule : public IAudioEditorModule
 {
 public:
@@ -123,6 +154,14 @@ public:
 		AddCustomClass("SoundfieldSubmix", FOnGetDetailCustomizationInstance::CreateStatic(&FSoundfieldSubmixDetailsCustomization::MakeInstance));
 
 		SetupIcons();
+
+		// Disable audio component asset broker in restricted mode 
+		// to disable CreateBlueprintUsing asset action
+		if (!IsRestrictedMode())
+		{
+			AudioComponentBroker = MakeShareable(new FAudioComponentBroker());
+			FComponentAssetBrokerage::RegisterBroker(AudioComponentBroker, UAudioComponent::StaticClass(), true, true);
+		}
 #if WITH_SNDFILE_IO
 		if (!Audio::SoundFileUtils::InitSoundFileIOManager())
 		{
@@ -168,6 +207,11 @@ public:
 			}
 			CustomClassLayoutNames.Reset();
 		}
+
+		if (UObjectInitialized() && !IsRestrictedMode())
+		{
+			FComponentAssetBrokerage::UnregisterBroker(AudioComponentBroker);
+		}
 	}
 
 	virtual void RegisterAssetActions() override
@@ -326,6 +370,16 @@ public:
 		return ImportObject<USoundWave>(SoundWavePackage, *InSoundWaveAssetName, RF_Public | RF_Standalone, *InWavFilename, nullptr, SoundWaveFactory);
 	}
 
+	bool IsRestrictedMode() const override
+	{
+		return bIsRestricted;
+	}
+
+	void SetRestrictedMode(bool bInRestricted) override
+	{
+		bIsRestricted = bInRestricted;
+	}
+
 private:
 
 	void SetupIcons()
@@ -385,6 +439,12 @@ private:
 	TSharedPtr<FGraphPanelPinConnectionFactory> SoundCueGraphConnectionFactory;
 	TSharedPtr<FGraphPanelPinConnectionFactory> SoundSubmixGraphConnectionFactory;
 	TSharedPtr<FSlateStyleSet> AudioStyleSet;
+
+	TSharedPtr<IComponentAssetBroker> AudioComponentBroker;
+
+	// Whether or not the editor is in restricted mode
+	// (to disable asset actions)
+	bool bIsRestricted;
 };
 
 IMPLEMENT_MODULE( FAudioEditorModule, AudioEditor );
diff --git a/Engine/Source/Editor/AudioEditor/Public/AudioEditorModule.h b/Engine/Source/Editor/AudioEditor/Public/AudioEditorModule.h
--- a/Engine/Source/Editor/AudioEditor/Public/AudioEditorModule.h
+++ b/Engine/Source/Editor/AudioEditor/Public/AudioEditorModule.h
@@ -2,9 +2,9 @@
 
 #pragma once
 
+#include "ComponentAssetBroker.h"
 #include "Containers/Array.h"
 #include "Containers/UnrealString.h"
-#include "CoreMinimal.h"
 #include "ISoundClassEditor.h"
 #include "ISoundCueEditor.h"
 #include "ISoundSubmixEditor.h"
@@ -35,6 +35,14 @@ AUDIOEDITOR_API DECLARE_LOG_CATEGORY_EXTERN(LogAudioEditor, Log, All);
 
 extern const FName AudioEditorAppIdentifier;
 
+class FAudioComponentBroker : public IComponentAssetBroker
+{
+public:
+	UClass* GetSupportedAssetClass() override;
+	virtual bool AssignAssetToComponent(UActorComponent* InComponent, UObject* InAsset) override;
+	virtual UObject* GetAssetFromComponent(UActorComponent* InComponent) override;
+};
+
 /** Sound class editor module interface */
 class IAudioEditorModule :	public IModuleInterface
 {
@@ -80,4 +88,10 @@ public:
 
 	/** Imports a sound wave from the given package. */
 	virtual USoundWave* ImportSoundWave(UPackage* const SoundWavePackage, const FString& InSoundWaveAssetName, const FString& InWavFilename) = 0;
+
+	// Queries if the Audio editor is in "restricted mode" 
+	virtual bool IsRestrictedMode() const = 0;
+
+	// Sets if the Audio editor is in "restricted mode" (certain asset actions may be disabled)
+	virtual void SetRestrictedMode(bool bInRestricted) = 0;
 };
diff --git a/Engine/Source/Editor/UnrealEd/Private/ComponentAssetBroker.cpp b/Engine/Source/Editor/UnrealEd/Private/ComponentAssetBroker.cpp
--- a/Engine/Source/Editor/UnrealEd/Private/ComponentAssetBroker.cpp
+++ b/Engine/Source/Editor/UnrealEd/Private/ComponentAssetBroker.cpp
@@ -124,43 +124,6 @@ public:
 	}
 };
 
-//////////////////////////////////////////////////////////////////////////
-// FAudioComponentBroker
-
-class FAudioComponentBroker : public IComponentAssetBroker
-{
-public:
-	UClass* GetSupportedAssetClass() override
-	{
-		return USoundBase::StaticClass();
-	}
-
-	virtual bool AssignAssetToComponent(UActorComponent* InComponent, UObject* InAsset) override
-	{
-		if (UAudioComponent* AudioComp = Cast<UAudioComponent>(InComponent))
-		{
-			USoundBase* Sound = Cast<USoundBase>(InAsset);
-
-			if ((Sound != NULL) || (InAsset == NULL))
-			{
-				AudioComp->SetSound(Sound);
-				return true;
-			}
-		}
-
-		return false;
-	}
-
-	virtual UObject* GetAssetFromComponent(UActorComponent* InComponent) override
-	{
-		if (UAudioComponent* AudioComp = Cast<UAudioComponent>(InComponent))
-		{
-			return AudioComp->Sound;
-		}
-		return NULL;
-	}
-};
-
 //////////////////////////////////////////////////////////////////////////
 // FChildActorComponentBroker
 
@@ -328,7 +291,6 @@ void FComponentAssetBrokerage::InitializeMap()
 		RegisterBroker(MakeShareable(new FStaticMeshComponentBroker), UStaticMeshComponent::StaticClass(), true, true);
 		RegisterBroker(MakeShareable(new FSkeletalMeshComponentBroker), USkeletalMeshComponent::StaticClass(), true, true);
 		RegisterBroker(MakeShareable(new FParticleSystemComponentBroker), UParticleSystemComponent::StaticClass(), true, true);
-		RegisterBroker(MakeShareable(new FAudioComponentBroker), UAudioComponent::StaticClass(), true, true);
 		RegisterBroker(MakeShareable(new FChildActorComponentBroker), UChildActorComponent::StaticClass(), true, false);
 	}
 }
