Bug Fix: Two cue points are created for loop regions with inconsistent data in each when USoundWave is imported/reimported

diff --git a/Engine/Source/Editor/AudioEditor/Private/Factories/SoundFactory.cpp b/Engine/Source/Editor/AudioEditor/Private/Factories/SoundFactory.cpp
--- a/Engine/Source/Editor/AudioEditor/Private/Factories/SoundFactory.cpp
+++ b/Engine/Source/Editor/AudioEditor/Private/Factories/SoundFactory.cpp
@@ -669,6 +669,8 @@ UObject* USoundFactory::CreateObject
 		// Setup the cue points
 		int TotalNum = WaveInfo.WaveCues.Num() + WaveInfo.WaveSampleLoops.Num();
 		Sound->CuePoints.Reset(TotalNum);
+		TSet<int32> WaveCuePointIDs;
+		WaveCuePointIDs.Reserve(WaveInfo.WaveCues.Num());
 
 		// Start with WaveCues
 		for (FWaveCue& WaveCue : WaveInfo.WaveCues)
@@ -679,9 +681,11 @@ UObject* USoundFactory::CreateObject
 			NewCuePoint.FramePosition = (int32)WaveCue.Position;
 			NewCuePoint.Label = WaveCue.Label;
 			Sound->CuePoints.Add(NewCuePoint);
+			WaveCuePointIDs.Add(NewCuePoint.CuePointID);
 		}
 
-		// add Sample Loops to end
+		// add Sample Loops to end or replace existing WaveCue 
+		// with matching SampleLoop to avoid duplicated/reduntant loops
 		bool FoundInvalidSampleLoops = false;
 		for (FWaveSampleLoop& SampleLoop : WaveInfo.WaveSampleLoops)
 		{
@@ -698,7 +702,25 @@ UObject* USoundFactory::CreateObject
 
 				FoundInvalidSampleLoops = true;
 			}
-			Sound->CuePoints.Add(NewCuePoint);
+
+			if (WaveCuePointIDs.Contains(NewCuePoint.CuePointID))
+			{
+				for (FSoundWaveCuePoint& CuePoint : Sound->CuePoints)
+				{
+					if (CuePoint.CuePointID == NewCuePoint.CuePointID)
+					{
+						//Initial CuePoint only has values for a marker
+						//and not a loop region, so loop values set here
+						CuePoint.bIsLoopRegion = true;
+						CuePoint.FrameLength = NewCuePoint.FrameLength;
+						break;
+					}
+				}
+			}
+			else
+			{
+				Sound->CuePoints.Add(NewCuePoint);
+			}
 		}
 
 		// fail import if we found invalid sample loops
