[Waveform Editor Bug] - Transforms are included as part of the Import from Sound Wave template

diff --git a/Engine/Source/Editor/AudioEditor/Private/Factories/SoundFactory.cpp b/Engine/Source/Editor/AudioEditor/Private/Factories/SoundFactory.cpp
--- a/Engine/Source/Editor/AudioEditor/Private/Factories/SoundFactory.cpp
+++ b/Engine/Source/Editor/AudioEditor/Private/Factories/SoundFactory.cpp
@@ -442,6 +442,7 @@ UObject* USoundFactory::CreateObject
 			Sound->NumChannels = 0;
 			Sound->ChannelOffsets.Reset();
 			Sound->ChannelSizes.Reset();
+			Sound->Transformations.Empty();
 		}
 
 		int32 SizeOfSample = (*WaveInfo.pBitsPerSample) / 8;
