Hide Pitch Modulation on Source Bus

diff --git a/Engine/Source/Editor/AudioEditor/Private/SoundModulationDestinationLayout.cpp b/Engine/Source/Editor/AudioEditor/Private/SoundModulationDestinationLayout.cpp
--- a/Engine/Source/Editor/AudioEditor/Private/SoundModulationDestinationLayout.cpp
+++ b/Engine/Source/Editor/AudioEditor/Private/SoundModulationDestinationLayout.cpp
@@ -1,6 +1,7 @@
 // Copyright Epic Games, Inc. All Rights Reserved.
 #include "SoundModulationDestinationLayout.h"
 
+#include "Algo/AnyOf.h"
 #include "AudioDevice.h"
 #include "AudioDeviceManager.h"
 #include "AudioEditorModule.h"
@@ -35,6 +36,7 @@
 #include "PropertyHandle.h"
 #include "SlotBase.h"
 #include "Sound/SoundModulationDestination.h"
+#include "Sound/SoundSourceBus.h"
 #include "Styling/AppStyle.h"
 #include "Templates/Casts.h"
 #include "Trace/Detail/Channel.h"
@@ -515,7 +517,19 @@ void FSoundModulationDefaultSettingsLayoutCustomization::CustomizeChildren(TShar
 		TSharedRef<IPropertyHandle> LowpassHandle = PropertyHandles.FindChecked(GET_MEMBER_NAME_CHECKED(FSoundModulationDefaultSettings, LowpassModulationDestination)).ToSharedRef();
 
 		ChildBuilder.AddProperty(VolumeHandle);
-		ChildBuilder.AddProperty(PitchHandle);
+
+		// Do not add the pitch destination to source buses
+		TArray<UObject*> Objects;
+		StructPropertyHandle->GetOuterObjects(Objects);
+		const bool bEditingSourceBus = Algo::AnyOf(Objects, [](const UObject* Object) -> bool
+			{
+				return Object->IsA<USoundSourceBus>();
+			});
+		if (!bEditingSourceBus)
+		{
+			ChildBuilder.AddProperty(PitchHandle);
+		}
+		
 		ChildBuilder.AddProperty(HighpassHandle);
 		ChildBuilder.AddProperty(LowpassHandle);
 	}
@@ -582,8 +596,19 @@ void FSoundModulationDefaultRoutingSettingsLayoutCustomization::CustomizeChildre
 
 		ChildBuilder.AddProperty(VolumeRouting);
 		ChildBuilder.AddProperty(VolumeHandle).Visibility(ShowModSettings(VolumeRouting));
-		ChildBuilder.AddProperty(PitchRouting);
-		ChildBuilder.AddProperty(PitchHandle).Visibility(ShowModSettings(PitchRouting));
+
+		// Do not add pitch destination to source buses
+		TArray<UObject*> Objects;
+		StructPropertyHandle->GetOuterObjects(Objects);
+		const bool bEditingSourceBus = Algo::AnyOf(Objects, [](const UObject* Object) -> bool
+			{
+				return Object->IsA<USoundSourceBus>();
+			});
+		if (!bEditingSourceBus)
+		{
+			ChildBuilder.AddProperty(PitchRouting);
+		}
+
 		ChildBuilder.AddProperty(HighpassRouting);
 		ChildBuilder.AddProperty(HighpassHandle).Visibility(ShowModSettings(HighpassRouting));
 		ChildBuilder.AddProperty(LowpassRouting);
diff --git a/Engine/Source/Runtime/Engine/Classes/Sound/SoundSourceBus.h b/Engine/Source/Runtime/Engine/Classes/Sound/SoundSourceBus.h
--- a/Engine/Source/Runtime/Engine/Classes/Sound/SoundSourceBus.h
+++ b/Engine/Source/Runtime/Engine/Classes/Sound/SoundSourceBus.h
@@ -69,6 +69,7 @@ public:
 	virtual void BeginCachePlatformData() override {}
 	virtual void FinishCachePlatformData() override {};	
 #endif	//WITH_EDITOR_ONLY_DATA
+	virtual void Serialize(FArchive& InArchive) override;
 	virtual void SerializeCookedPlatformData(FArchive& Ar) override {}
 	virtual bool IsSeekable() const override { return false; }
 	//~ End USoundWave Interface
diff --git a/Engine/Source/Runtime/Engine/Private/SoundSourceBus.cpp b/Engine/Source/Runtime/Engine/Private/SoundSourceBus.cpp
--- a/Engine/Source/Runtime/Engine/Private/SoundSourceBus.cpp
+++ b/Engine/Source/Runtime/Engine/Private/SoundSourceBus.cpp
@@ -50,6 +50,20 @@ void USoundSourceBus::PostEditChangeProperty(FPropertyChangedEvent& PropertyChan
 }
 #endif
 
+void USoundSourceBus::Serialize(FArchive& InArchive)
+{
+	Super::Serialize(InArchive);
+	
+#if WITH_EDITORONLY_DATA
+	if (InArchive.IsLoading())
+	{
+		// Source Buses do not support pitch modulation, so remove any connections if users previously attempted it.
+		ModulationSettings.PitchRouting = EModulationRouting::Disable;
+		ModulationSettings.PitchModulationDestination.Modulators.Empty();
+	}
+#endif // WITH_EDITORONLY_DATA
+}
+
 bool USoundSourceBus::IsPlayable() const
 {
 	return true;
