[Waveform Editor] FIxed make waveform editor default asset editor for SoundWaves

diff --git a/Engine/Plugins/Editor/WaveformEditor/Source/WaveformEditor/Private/WaveformEditor.cpp b/Engine/Plugins/Editor/WaveformEditor/Source/WaveformEditor/Private/WaveformEditor.cpp
--- a/Engine/Plugins/Editor/WaveformEditor/Source/WaveformEditor/Private/WaveformEditor.cpp
+++ b/Engine/Plugins/Editor/WaveformEditor/Source/WaveformEditor/Private/WaveformEditor.cpp
@@ -1070,9 +1070,9 @@ const TSharedRef<FTabManager::FLayout> FWaveformEditor::SetupStandaloneLayout()
 				(
 					FTabManager::NewStack()
 					->SetSizeCoefficient(0.2f)
-					->AddTab(TransformationsTabId, ETabState::OpenedTab)
 					->AddTab(PropertiesTabId, ETabState::OpenedTab)
-					->SetForegroundTab(TransformationsTabId)
+					->AddTab(TransformationsTabId, ETabState::OpenedTab)
+					->SetForegroundTab(PropertiesTabId)
 
 				)
 				->Split
diff --git a/Engine/Plugins/Editor/WaveformEditor/Source/WaveformEditor/Private/WaveformEditorInstantiator.cpp b/Engine/Plugins/Editor/WaveformEditor/Source/WaveformEditor/Private/WaveformEditorInstantiator.cpp
--- a/Engine/Plugins/Editor/WaveformEditor/Source/WaveformEditor/Private/WaveformEditorInstantiator.cpp
+++ b/Engine/Plugins/Editor/WaveformEditor/Source/WaveformEditor/Private/WaveformEditorInstantiator.cpp
@@ -3,6 +3,7 @@
 #include "WaveformEditorInstantiator.h"
 
 #include "Algo/AnyOf.h"
+#include "AudioEditorModule.h"
 #include "ContentBrowserMenuContexts.h"
 #include "EditorReimportHandler.h"
 #include "Framework/Application/SlateApplication.h"
@@ -63,6 +64,12 @@ void FWaveformEditorInstantiator::ExtendContentBrowserSelectionMenu()
 		}));
 }
 
+void FWaveformEditorInstantiator::RegisterAsSoundwaveEditor()
+{
+	IAudioEditorModule& AudioEditorModule = FModuleManager::LoadModuleChecked<IAudioEditorModule>( "AudioEditor" );
+	AudioEditorModule.SoundWaveEditorOpen.BindSP(this, &FWaveformEditorInstantiator::CreateWaveformEditor);
+}
+
 void FWaveformEditorInstantiator::ExecuteCreateWaveformEditor(const FToolMenuContext& MenuContext)
 {
 	if (const UContentBrowserAssetContextMenuContext* Context = UContentBrowserAssetContextMenuContext::FindContextWithAssets(MenuContext))
diff --git a/Engine/Plugins/Editor/WaveformEditor/Source/WaveformEditor/Private/WaveformEditorInstantiator.h b/Engine/Plugins/Editor/WaveformEditor/Source/WaveformEditor/Private/WaveformEditorInstantiator.h
--- a/Engine/Plugins/Editor/WaveformEditor/Source/WaveformEditor/Private/WaveformEditorInstantiator.h
+++ b/Engine/Plugins/Editor/WaveformEditor/Source/WaveformEditor/Private/WaveformEditorInstantiator.h
@@ -12,6 +12,9 @@ class FWaveformEditorInstantiator : public IWaveformEditorInstantiator,  public
 {
 public:
 	virtual void ExtendContentBrowserSelectionMenu() override;
+	
+	// Registers using a delegate as the default editor for soundwave assets without main engine referencing waveform plugin in Sound Wave Asset Definition
+	virtual void RegisterAsSoundwaveEditor() override;
 	virtual ~FWaveformEditorInstantiator() = default;
 
 private:
diff --git a/Engine/Plugins/Editor/WaveformEditor/Source/WaveformEditor/Private/WaveformEditorModule.cpp b/Engine/Plugins/Editor/WaveformEditor/Source/WaveformEditor/Private/WaveformEditorModule.cpp
--- a/Engine/Plugins/Editor/WaveformEditor/Source/WaveformEditor/Private/WaveformEditorModule.cpp
+++ b/Engine/Plugins/Editor/WaveformEditor/Source/WaveformEditor/Private/WaveformEditorModule.cpp
@@ -19,6 +19,7 @@ void FWaveformEditorModule::StartupModule()
 
 	WaveformEditorInstantiator = MakeShared<FWaveformEditorInstantiator>();
 	RegisterContentBrowserExtensions(WaveformEditorInstantiator.Get());
+	WaveformEditorInstantiator->RegisterAsSoundwaveEditor();
 }
 
 void FWaveformEditorModule::ShutdownModule()
diff --git a/Engine/Plugins/Editor/WaveformEditor/Source/WaveformEditor/Public/IWaveformEditorInstantiator.h b/Engine/Plugins/Editor/WaveformEditor/Source/WaveformEditor/Public/IWaveformEditorInstantiator.h
--- a/Engine/Plugins/Editor/WaveformEditor/Source/WaveformEditor/Public/IWaveformEditorInstantiator.h
+++ b/Engine/Plugins/Editor/WaveformEditor/Source/WaveformEditor/Public/IWaveformEditorInstantiator.h
@@ -10,6 +10,7 @@ class IWaveformEditorInstantiator
 {
 public:
 	virtual void ExtendContentBrowserSelectionMenu() = 0;
+	virtual void RegisterAsSoundwaveEditor() = 0;
 
 private:
 	virtual void CreateWaveformEditor(TArray<USoundWave*> SoundWavesToEdit) = 0;
diff --git a/Engine/Source/Editor/AudioEditor/Private/AssetTypeActions/AssetDefinition_SoundWave.cpp b/Engine/Source/Editor/AudioEditor/Private/AssetTypeActions/AssetDefinition_SoundWave.cpp
--- a/Engine/Source/Editor/AudioEditor/Private/AssetTypeActions/AssetDefinition_SoundWave.cpp
+++ b/Engine/Source/Editor/AudioEditor/Private/AssetTypeActions/AssetDefinition_SoundWave.cpp
@@ -47,10 +47,20 @@ TConstArrayView<FAssetCategoryPath> UAssetDefinition_SoundWave::GetAssetCategori
 EAssetCommandResult UAssetDefinition_SoundWave::OpenAssets(const FAssetOpenArgs& OpenArgs) const
 {
 	IAudioEditorModule& AudioEditorModule = FModuleManager::LoadModuleChecked<IAudioEditorModule>( "AudioEditor" );
-	for (USoundWave* SoundWave : OpenArgs.LoadObjects<USoundWave>())
+
+	// Use a custom editor for soundwaves if one is registered
+	if (AudioEditorModule.SoundWaveEditorOpen.IsBound())
+	{
+		AudioEditorModule.SoundWaveEditorOpen.Execute(OpenArgs.LoadObjects<USoundWave>());
+	}
+	else // Otherwise open a default editor
 	{
-		FSimpleAssetEditor::CreateEditor(EToolkitMode::Standalone, OpenArgs.ToolkitHost, SoundWave);
+		for (USoundWave* SoundWave : OpenArgs.LoadObjects<USoundWave>())
+		{
+			FSimpleAssetEditor::CreateEditor(EToolkitMode::Standalone, OpenArgs.ToolkitHost, SoundWave);
+		}
 	}
+
 	return EAssetCommandResult::Handled;
 }
 
diff --git a/Engine/Source/Editor/AudioEditor/Public/AudioEditorModule.h b/Engine/Source/Editor/AudioEditor/Public/AudioEditorModule.h
--- a/Engine/Source/Editor/AudioEditor/Public/AudioEditorModule.h
+++ b/Engine/Source/Editor/AudioEditor/Public/AudioEditorModule.h
@@ -94,4 +94,9 @@ public:
 
 	// Sets if the Audio editor is in "restricted mode" (certain asset actions may be disabled)
 	virtual void SetRestrictedMode(bool bInRestricted) = 0;
+
+	// Allows systems to register to open soundwave files with custom editor
+	DECLARE_DELEGATE_OneParam(SoundWaveEditorOpenDelegate, TArray<USoundWave*>);
+
+	SoundWaveEditorOpenDelegate SoundWaveEditorOpen;
 };
