Deprecate GetModulationParameter function in favor of one which returns a const pointer

diff --git a/Engine/Plugins/Runtime/AudioModulation/Source/AudioModulation/Private/SoundModulationGeneratorProxy.cpp b/Engine/Plugins/Runtime/AudioModulation/Source/AudioModulation/Private/SoundModulationGeneratorProxy.cpp
--- a/Engine/Plugins/Runtime/AudioModulation/Source/AudioModulation/Private/SoundModulationGeneratorProxy.cpp
+++ b/Engine/Plugins/Runtime/AudioModulation/Source/AudioModulation/Private/SoundModulationGeneratorProxy.cpp
@@ -12,7 +12,7 @@ namespace AudioModulation
 	const Audio::FModulationParameter& FModulationGeneratorSettings::GetOutputParameter() const
 	{
 		// For now, all generators are normalized values (0, 1).
-		return Audio::GetModulationParameter({ });
+		return Audio::GetDefaultModulationParameter();
 	}
 
 	Audio::FModulatorTypeId FModulationGeneratorSettings::Register(Audio::FModulatorHandleId HandleId, IAudioModulationManager& InModulation) const
diff --git a/Engine/Plugins/Runtime/AudioModulation/Source/AudioModulation/Private/SoundModulationParameter.cpp b/Engine/Plugins/Runtime/AudioModulation/Source/AudioModulation/Private/SoundModulationParameter.cpp
--- a/Engine/Plugins/Runtime/AudioModulation/Source/AudioModulation/Private/SoundModulationParameter.cpp
+++ b/Engine/Plugins/Runtime/AudioModulation/Source/AudioModulation/Private/SoundModulationParameter.cpp
@@ -372,7 +372,11 @@ namespace AudioModulation
 		}
 
 		// Returns default modulation parameter if no parameter provided.
-		return Audio::GetModulationParameter(ParamName);
+		if (const Audio::FModulationParameter* Parameter = Audio::GetModulationParameterPtr(ParamName))
+		{
+			return *Parameter;
+		}
+		return Audio::GetDefaultModulationParameter();
 	}
 
 	FSoundModulationPluginParameterAssetProxy::FSoundModulationPluginParameterAssetProxy(USoundModulationParameter* InParameter)
diff --git a/Engine/Plugins/Runtime/AudioModulation/Source/AudioModulationEditor/Private/Layouts/SoundControlBusMixStageLayout.cpp b/Engine/Plugins/Runtime/AudioModulation/Source/AudioModulationEditor/Private/Layouts/SoundControlBusMixStageLayout.cpp
--- a/Engine/Plugins/Runtime/AudioModulation/Source/AudioModulationEditor/Private/Layouts/SoundControlBusMixStageLayout.cpp
+++ b/Engine/Plugins/Runtime/AudioModulation/Source/AudioModulationEditor/Private/Layouts/SoundControlBusMixStageLayout.cpp
@@ -103,7 +103,7 @@ namespace AudioModulationEditorUtils
 		}
 		else
 		{
-			const Audio::FModulationParameter& DefaultParam = Audio::GetModulationParameter({ });
+			const Audio::FModulationParameter& DefaultParam = Audio::GetDefaultModulationParameter();
 
 			float NormalizedValue = 1.0f;
 			DefaultParam.GetDefaultNormalizedConversionFunction()(NormalizedValue);
diff --git a/Engine/Source/Editor/AudioEditor/Private/SoundModulationDestinationLayout.cpp b/Engine/Source/Editor/AudioEditor/Private/SoundModulationDestinationLayout.cpp
--- a/Engine/Source/Editor/AudioEditor/Private/SoundModulationDestinationLayout.cpp
+++ b/Engine/Source/Editor/AudioEditor/Private/SoundModulationDestinationLayout.cpp
@@ -179,19 +179,20 @@ namespace ModDestinationLayoutUtils
 		if (OutParamName != FName())
 		{
 			// If parameter was provided, it overrides ClampMin/Max.
-			const Audio::FModulationParameter& Parameter = Audio::GetModulationParameter(OutParamName);
-
-			// if no valid parameter was found & the user has specified their own clamping, don't override it with the default parameter range of [0,1]
-			if (false == (Parameter.ParameterName == FName() && bClampValuesSet))
+			if (const Audio::FModulationParameter* Parameter = Audio::GetModulationParameterPtr(OutParamName))
 			{
-				UIMinValue = Parameter.MinValue;
-				UIMaxValue = Parameter.MaxValue;
-				ClampMinValue = UIMinValue;
-				ClampMaxValue = UIMaxValue;
-				OutUnitDisplayText = Parameter.UnitDisplayName;
-				if (bClampValuesSet)
+				// if no valid parameter was found & the user has specified their own clamping, don't override it with the default parameter range of [0,1]
+				if (false == (Parameter->ParameterName == FName() && bClampValuesSet))
 				{
-					UE_LOG(LogAudioEditor, Verbose, TEXT("ClampMin/Max overridden by AudioModulation plugin asset with ParamName '%s'."), *OutParamName.ToString());
+					UIMinValue = Parameter->MinValue;
+					UIMaxValue = Parameter->MaxValue;
+					ClampMinValue = UIMinValue;
+					ClampMaxValue = UIMaxValue;
+					OutUnitDisplayText = Parameter->UnitDisplayName;
+					if (bClampValuesSet)
+					{
+						UE_LOG(LogAudioEditor, Verbose, TEXT("ClampMin/Max overridden by AudioModulation plugin asset with ParamName '%s'."), *OutParamName.ToString());
+					}
 				}
 			}
 		}
@@ -290,16 +291,19 @@ namespace ModDestinationLayoutUtils
 						float CurrentValue = 0.0f;
 						ValueHandle->GetValue(CurrentValue);
 
-						const Audio::FModulationParameter& Parameter = Audio::GetModulationParameter(ParamName);
-						return Parameter.DefaultValue == CurrentValue
-							? EVisibility::Hidden
-							: EVisibility::Visible;
+						const Audio::FModulationParameter* Parameter = Audio::GetModulationParameterPtr(ParamName);
+						if (Parameter && Parameter->DefaultValue == CurrentValue)
+						{
+							return EVisibility::Hidden;
+						}
+						return EVisibility::Visible;
 					}))
 					.OnClicked(FOnClicked::CreateLambda([ParamName, ValueHandle]()
 					{
-						Audio::FModulationParameter Parameter = Audio::GetModulationParameter(ParamName);
-						ValueHandle->SetValue(Parameter.DefaultValue);
-
+						if (const Audio::FModulationParameter* Parameter = Audio::GetModulationParameterPtr(ParamName))
+						{
+							ValueHandle->SetValue(Parameter->DefaultValue);
+						}
 						return FReply::Handled();
 					}))
 					.Content()
@@ -317,13 +321,13 @@ namespace ModDestinationLayoutUtils
 			if (bEnabled)
 			{
 				const FName ParamName = ModDestinationLayoutUtils::GetParameterNameFromMetaData(StructPropertyHandle);
-				const Audio::FModulationParameter& Parameter = Audio::GetModulationParameter(ParamName);
-				ValueHandle->SetValue(Parameter.DefaultValue);
-			}
-			else
-			{
-				ModulatorsHandle->AsSet()->Empty();
+				if (const Audio::FModulationParameter* Parameter = Audio::GetModulationParameterPtr(ParamName))
+				{
+					ValueHandle->SetValue(Parameter->DefaultValue);
+					return;
+				}
 			}
+			ModulatorsHandle->AsSet()->Empty();
 		}));
 	}
 
diff --git a/Engine/Source/Runtime/AudioExtensions/Private/IAudioModulation.cpp b/Engine/Source/Runtime/AudioExtensions/Private/IAudioModulation.cpp
--- a/Engine/Source/Runtime/AudioExtensions/Private/IAudioModulation.cpp
+++ b/Engine/Source/Runtime/AudioExtensions/Private/IAudioModulation.cpp
@@ -12,13 +12,15 @@ namespace Audio
 	namespace ModulationPrivate
 	{
 		static std::atomic<FModulatorHandleId> NextHandleId = INDEX_NONE;
+
+		static const FModulationParameter DefaultParameter = { };
 	}
 
 	namespace ModulationInterfacePrivate
 	{
 		class FModulationParameterRegistry
 		{
-			TMap<FName, FModulationParameter> Values;
+			TMap<FName, TUniquePtr<FModulationParameter>> Values;
 
 			mutable FCriticalSection ThreadSafeValueAccessor;
 
@@ -32,14 +34,16 @@ namespace Audio
 			void Register(FName InName, FModulationParameter&& InParameter)
 			{
 				FScopeLock Lock(&ThreadSafeValueAccessor);
-				if (FModulationParameter* Value = Values.Find(InName))
-				{
-					*Value = MoveTemp(InParameter);
-				}
-				else
+				if (const TUniquePtr<FModulationParameter>* ParamPtr = Values.Find(InName))
 				{
-					Values.Add(InName, MoveTemp(InParameter));
+					if (FModulationParameter* Param = ParamPtr->Get())
+					{
+						*Param = MoveTemp(InParameter);
+						return;
+					}
 				}
+				
+				Values.Add(InName, MakeUnique<FModulationParameter>(FModulationParameter(InParameter)));
 			}
 
 			bool Unregister(FName InName)
@@ -54,17 +58,19 @@ namespace Audio
 				Values.Reset();
 			}
 
-			const FModulationParameter& Get(FName InName) const
+			const FModulationParameter* Get(FName InName) const
 			{
 				FScopeLock Lock(&ThreadSafeValueAccessor);
 
-				static const FModulationParameter DefaultParameter { };
-				if (const FModulationParameter* Param = Values.Find(InName))
+				if (const TUniquePtr<FModulationParameter>* ParamPtr = Values.Find(InName))
 				{
-					return *Param;
+					if (const FModulationParameter* Param = ParamPtr->Get())
+					{
+						return Param;
+					}
 				}
 
-				return DefaultParameter;
+				return nullptr;
 			}
 		} ParameterRegistry;
 	}
@@ -204,11 +210,28 @@ namespace Audio
 		using namespace ModulationInterfacePrivate;
 		return ParameterRegistry.IsRegistered(InName);
 	}
+	
+	const FModulationParameter* GetModulationParameterPtr(FName InName)
+	{
+		using namespace ModulationInterfacePrivate;
+		return ParameterRegistry.Get(InName);
+	}
+
+	const FModulationParameter& GetDefaultModulationParameter()
+	{
+		return ModulationPrivate::DefaultParameter;
+	}
 
 	const FModulationParameter& GetModulationParameter(FName InName)
 	{
 		using namespace ModulationInterfacePrivate;
-		return ParameterRegistry.Get(InName);
+
+		if (const FModulationParameter* Param = ParameterRegistry.Get(InName))
+		{
+			return *Param;
+		}
+		
+		return ModulationPrivate::DefaultParameter;
 	}
 
 	FModulatorHandle::FModulatorHandle(Audio::FModulationParameter&& InParameter)
@@ -380,7 +403,7 @@ namespace Audio
 
 const Audio::FModulationParameter& USoundModulatorBase::GetOutputParameter() const
 {
-	return Audio::GetModulationParameter({ });
+	return Audio::ModulationPrivate::DefaultParameter;
 }
 
 TSharedPtr<Audio::IProxyData> USoundModulatorBase::CreateProxyData(const Audio::FProxyDataInitParams& InitParams)
diff --git a/Engine/Source/Runtime/AudioExtensions/Public/IAudioModulation.h b/Engine/Source/Runtime/AudioExtensions/Public/IAudioModulation.h
--- a/Engine/Source/Runtime/AudioExtensions/Public/IAudioModulation.h
+++ b/Engine/Source/Runtime/AudioExtensions/Public/IAudioModulation.h
@@ -102,6 +102,10 @@ namespace Audio
 	AUDIOEXTENSIONS_API void RegisterModulationParameter(FName InName, FModulationParameter&& InParameter);
 	AUDIOEXTENSIONS_API bool UnregisterModulationParameter(FName InName);
 	AUDIOEXTENSIONS_API void UnregisterAllModulationParameters();
+	AUDIOEXTENSIONS_API const FModulationParameter* GetModulationParameterPtr(FName InName);
+	AUDIOEXTENSIONS_API const FModulationParameter& GetDefaultModulationParameter();
+	
+	UE_DEPRECATED(5.6, "Please use the version of GetModulationParameter which returns a pointer instead of a reference, and check that it is valid in your code.")
 	AUDIOEXTENSIONS_API const FModulationParameter& GetModulationParameter(FName InName);
 
 	/** Interface for cached off Modulator UObject data used as default settings to
diff --git a/Engine/Source/Runtime/Engine/Private/Audio.cpp b/Engine/Source/Runtime/Engine/Private/Audio.cpp
--- a/Engine/Source/Runtime/Engine/Private/Audio.cpp
+++ b/Engine/Source/Runtime/Engine/Private/Audio.cpp
@@ -102,6 +102,14 @@ FAutoConsoleVariableRef CVarMinVolumeThreshold(
 	TEXT("Default is 0.0001 (-80 dB)"),
 	ECVF_Default);
 
+namespace AudioModulationPrivate
+{
+	static const FLazyName VolumeParameterName("Volume");
+	static const FLazyName PitchParameterName("Pitch");
+	static const FLazyName HighpassParameterName("HPFCutoffFrequency");
+	static const FLazyName LowpassParameterName("LPFCutoffFrequency");
+}
+
 bool IsAudioPluginEnabled(EAudioPlugin PluginType)
 {
 	switch (PluginType)
@@ -200,17 +208,33 @@ Audio::FModulationParameter FWaveInstance::GetParameterBasedOnDestination(const
 	switch (InDestination)
 	{
 	case EModulationDestination::Volume:
-		return Audio::GetModulationParameter("Volume");
+		if (const Audio::FModulationParameter* Parameter = Audio::GetModulationParameterPtr(AudioModulationPrivate::VolumeParameterName))
+		{
+			return *Parameter;
+		}
+		break;
 	case EModulationDestination::Pitch:
-		return Audio::GetModulationParameter("Pitch");
+		if (const Audio::FModulationParameter* Parameter = Audio::GetModulationParameterPtr(AudioModulationPrivate::PitchParameterName))
+		{
+			return *Parameter;
+		}
+		break;
 	case EModulationDestination::Highpass:
-		return Audio::GetModulationParameter("HPFCutoffFilter");
+		if (const Audio::FModulationParameter* Parameter = Audio::GetModulationParameterPtr(AudioModulationPrivate::HighpassParameterName))
+		{
+			return *Parameter;
+		}
+		break;
 	case EModulationDestination::Lowpass:
-		return Audio::GetModulationParameter("LPFCutoffFilter");
+		if (const Audio::FModulationParameter* Parameter = Audio::GetModulationParameterPtr(AudioModulationPrivate::LowpassParameterName))
+		{
+			return *Parameter;
+		}
+		break;
 	default:
 		break;
 	}
-	return Audio::FModulationParameter();
+	return Audio::GetDefaultModulationParameter();
 }
 
 void FWaveInstance::CalculateInheritedModulationSettingsAndValue(const EModulationDestination InDestination, FSoundModulationDestinationSettings& InOutSettings, float& OutValue) const
diff --git a/Engine/Source/Runtime/Engine/Private/SoundModulationDestination.cpp b/Engine/Source/Runtime/Engine/Private/SoundModulationDestination.cpp
--- a/Engine/Source/Runtime/Engine/Private/SoundModulationDestination.cpp
+++ b/Engine/Source/Runtime/Engine/Private/SoundModulationDestination.cpp
@@ -131,7 +131,15 @@ namespace Audio
 		DestinationData->bValueNormalized = bInValueNormalized;
 
 		DestinationData->OutputBuffer.Reset();
-		DestinationData->Parameter = Audio::GetModulationParameter(InParameterName);
+		if (const FModulationParameter* Param = GetModulationParameterPtr(InParameterName))
+		{
+			DestinationData->Parameter = *Param;
+		}
+		else
+		{
+			UE_LOG(LogAudio, Warning, TEXT("Modulation Parameter %s not found while Initializing Modulation Destination. Using default parameter instead."), *InParameterName.ToString());
+			DestinationData->Parameter = Audio::GetDefaultModulationParameter();
+		}
 
 		DestinationData->ResetHandles();
 	}
