Fix Pitch Destination being hidden in wrong cases

diff --git a/Engine/Source/Editor/AudioEditor/Private/SoundModulationDestinationLayout.cpp b/Engine/Source/Editor/AudioEditor/Private/SoundModulationDestinationLayout.cpp
--- a/Engine/Source/Editor/AudioEditor/Private/SoundModulationDestinationLayout.cpp
+++ b/Engine/Source/Editor/AudioEditor/Private/SoundModulationDestinationLayout.cpp
@@ -611,6 +611,7 @@ void FSoundModulationDefaultRoutingSettingsLayoutCustomization::CustomizeChildre
 		if (!bEditingSourceBus)
 		{
 			ChildBuilder.AddProperty(PitchRouting);
+			ChildBuilder.AddProperty(PitchHandle).Visibility(ShowModSettings(PitchRouting));
 		}
 
 		ChildBuilder.AddProperty(HighpassRouting);
