Fix Submixes not connecting automatically when dragging off pin to create a new one

diff --git a/Engine/Source/Editor/AudioEditor/Private/SoundSubmixGraphNode.cpp b/Engine/Source/Editor/AudioEditor/Private/SoundSubmixGraphNode.cpp
--- a/Engine/Source/Editor/AudioEditor/Private/SoundSubmixGraphNode.cpp
+++ b/Engine/Source/Editor/AudioEditor/Private/SoundSubmixGraphNode.cpp
@@ -177,11 +177,11 @@ void USoundSubmixGraphNode::AutowireNewNode(UEdGraphPin* FromPin)
 
 		if (FromPin->Direction == EGPD_Input)
 		{
-			Schema->TryCreateConnection(FromPin, ChildPin);
+			Schema->TryCreateConnection(FromPin, ParentPin);
 		}
 		else
 		{
-			Schema->TryCreateConnection(FromPin, ParentPin);
+			Schema->TryCreateConnection(FromPin, ChildPin);
 		}
 	}
 }
