Disable editor pref UI for audio device selection on unsupported platforms.

#rnx

diff --git a/Engine/Source/Editor/AudioEditor/Classes/AudioEditorSettings.h b/Engine/Source/Editor/AudioEditor/Classes/AudioEditorSettings.h
--- a/Engine/Source/Editor/AudioEditor/Classes/AudioEditorSettings.h
+++ b/Engine/Source/Editor/AudioEditor/Classes/AudioEditorSettings.h
@@ -95,6 +95,8 @@ public:
 	virtual void CustomizeDetails(IDetailLayoutBuilder& DetailBuilder) override;
 	//~ End IDetailCustomization
 
+	static bool IsUIEnabled();
+
 private:
 	TSharedPtr<IPropertyHandle> UseSystemDevicePropertyHandle;
 	TSharedPtr<IPropertyHandle> AudioDeviceIdPropertyHandle;
diff --git a/Engine/Source/Editor/AudioEditor/Private/AudioEditorSettings.cpp b/Engine/Source/Editor/AudioEditor/Private/AudioEditorSettings.cpp
--- a/Engine/Source/Editor/AudioEditor/Private/AudioEditorSettings.cpp
+++ b/Engine/Source/Editor/AudioEditor/Private/AudioEditorSettings.cpp
@@ -129,6 +129,7 @@ void FAudioOutputDeviceCustomization::CustomizeDetails(IDetailLayoutBuilder& Det
 		.ValueContent()
 		[
 			SNew(SCheckBox)
+			.IsEnabled(IsUIEnabled())
 			.IsChecked_Lambda([this]() -> ECheckBoxState
 			{
 				bool bUseSystemDefault;
@@ -203,7 +204,7 @@ void FAudioOutputDeviceCustomization::CustomizeDetails(IDetailLayoutBuilder& Det
 		{
 			bool bUseSystemDefault;
 			UseSystemDevicePropertyHandle->GetValue(bUseSystemDefault);
-			return !bUseSystemDefault;
+			return IsUIEnabled() && !bUseSystemDefault;
 		})
 		.Text_Lambda([this]()
 		{
@@ -221,6 +222,16 @@ void FAudioOutputDeviceCustomization::CustomizeDetails(IDetailLayoutBuilder& Det
 	}
 }
 
+bool FAudioOutputDeviceCustomization::IsUIEnabled()
+{
+	// Only enable UI on Windows platform
+#if PLATFORM_WINDOWS
+	return true;
+#else
+	return false;
+#endif
+}
+
 TArray<FAudioPlatformDeviceInfo> FAudioOutputDeviceCustomization::GetAvailableAudioOutputDevices()
 {
 	TArray<FAudioPlatformDeviceInfo> OutputDevices;
