Fix uninitialized variable bIsRestricted

#rnx

diff --git a/Engine/Source/Editor/AudioEditor/Private/AudioEditorModule.cpp b/Engine/Source/Editor/AudioEditor/Private/AudioEditorModule.cpp
--- a/Engine/Source/Editor/AudioEditor/Private/AudioEditorModule.cpp
+++ b/Engine/Source/Editor/AudioEditor/Private/AudioEditorModule.cpp
@@ -459,7 +459,7 @@ private:
 
 	// Whether or not the editor is in restricted mode
 	// (to disable asset actions)
-	bool bIsRestricted;
+	bool bIsRestricted = false;
 };
 
 IMPLEMENT_MODULE( FAudioEditorModule, AudioEditor );
