[UDN 17088530] Assert loading a module via PostLoad while a UDeveloperSettings-based config property is visible
- Had the AudioExtensionsModule get loaded in the AudioEditorModule startup so that it doesn't have to rely on getting loaded via AutoDiscoveredSettings since it may be needed before it can be reloaded and the AudioEditor uses the custom classes from the AudioExtension as well.

diff --git a/Engine/Source/Editor/AudioEditor/Private/AudioEditorModule.cpp b/Engine/Source/Editor/AudioEditor/Private/AudioEditorModule.cpp
--- a/Engine/Source/Editor/AudioEditor/Private/AudioEditorModule.cpp
+++ b/Engine/Source/Editor/AudioEditor/Private/AudioEditorModule.cpp
@@ -150,6 +150,10 @@ public:
 			InPropertyModule->RegisterCustomClassLayout(Name, InstanceGetter);
 			CustomClassLayoutNames.Add(Name);
 		};
+
+		// Load the module for Audio Extensions so it doesn't rely on AutoDiscoveredSettings which can get unregistered since it is utilized by the following custom classes
+		FModuleManager::Get().LoadModule(TEXT("AudioExtensions"));
+
 		AddCustomClass("EndpointSubmix", FOnGetDetailCustomizationInstance::CreateStatic(&FEndpointSubmixDetailsCustomization::MakeInstance));
 		AddCustomClass("SoundfieldEndpointSubmix", FOnGetDetailCustomizationInstance::CreateStatic(&FSoundfieldEndpointSubmixDetailsCustomization::MakeInstance));
 		AddCustomClass("SoundfieldSubmix", FOnGetDetailCustomizationInstance::CreateStatic(&FSoundfieldSubmixDetailsCustomization::MakeInstance));
