Adding platform Audio API menu to audio editor settings

#rnx

diff --git a/Engine/Config/Windows/WindowsEditorSettings.ini b/Engine/Config/Windows/WindowsEditorSettings.ini
new file mode 100644
--- /dev/null
+++ b/Engine/Config/Windows/WindowsEditorSettings.ini
@@ -0,0 +1,3 @@
+[/Script/AudioEditor.AudioEditorSettings]
++AudioMixerPlatforms=(DisplayName="XAudio2",ModuleName="AudioMixerXAudio2")
++AudioMixerPlatforms=(DisplayName="Wasapi",ModuleName="AudioMixerWasapi")
diff --git a/Engine/Source/Editor/AudioEditor/AudioEditor.Build.cs b/Engine/Source/Editor/AudioEditor/AudioEditor.Build.cs
--- a/Engine/Source/Editor/AudioEditor/AudioEditor.Build.cs
+++ b/Engine/Source/Editor/AudioEditor/AudioEditor.Build.cs
@@ -86,6 +86,7 @@ public class AudioEditor : ModuleRules
 			RuntimeDependencies.Add("$(EngineDir)/Binaries/ThirdParty/libsndfile/" + PlatformName + "/libsndfile-1.dll");
 
 			PublicDefinitions.Add("WITH_SNDFILE_IO=1");
+			PublicDefinitions.Add("ENABLE_DEVICE_SETTINGS=1");
 		}
 		else if (Target.IsInPlatformGroup(UnrealPlatformGroup.Apple))
 		{
@@ -99,6 +100,7 @@ public class AudioEditor : ModuleRules
 			RuntimeDependencies.Add(LibSndFilePath);
 			
 			PublicDefinitions.Add("WITH_SNDFILE_IO=1");
+			PublicDefinitions.Add("ENABLE_DEVICE_SETTINGS=0");
 		}
 		else if (Target.IsInPlatformGroup(UnrealPlatformGroup.Linux))
 		{
@@ -112,10 +114,12 @@ public class AudioEditor : ModuleRules
 			RuntimeDependencies.Add(LibSndFilePath);
 			
 			PublicDefinitions.Add("WITH_SNDFILE_IO=1");
+			PublicDefinitions.Add("ENABLE_DEVICE_SETTINGS=0");
 		}
 		else
 		{
 			PublicDefinitions.Add("WITH_SNDFILE_IO=0");
+			PublicDefinitions.Add("ENABLE_DEVICE_SETTINGS=0");
 		}
 	}
 }
diff --git a/Engine/Source/Editor/AudioEditor/Classes/AudioEditorSettings.h b/Engine/Source/Editor/AudioEditor/Classes/AudioEditorSettings.h
--- a/Engine/Source/Editor/AudioEditor/Classes/AudioEditorSettings.h
+++ b/Engine/Source/Editor/AudioEditor/Classes/AudioEditorSettings.h
@@ -23,7 +23,31 @@ namespace Audio
 class UObject;
 struct FPropertyChangedEvent;
 
-UCLASS(config = EditorPerProjectUserSettings, defaultconfig, meta = (DisplayName = "Audio"))
+USTRUCT()
+struct FAudioPlatform
+{
+	GENERATED_BODY()
+	
+	UPROPERTY(config)
+	FString DisplayName;
+	
+	UPROPERTY(config)
+	FString ModuleName;
+
+	FAudioPlatform() = default;
+
+	FAudioPlatform(const FString& InDisplayName, const FString& InModuleName)
+		: DisplayName(InDisplayName)
+		, ModuleName(InModuleName)
+	{}
+
+	bool operator<(const FAudioPlatform& Other) const
+	{
+		return DisplayName < Other.DisplayName;
+	}
+};
+
+UCLASS(config = EditorSettings, meta = (DisplayName = "Audio"))
 class AUDIOEDITOR_API UAudioEditorSettings : public UDeveloperSettings
 {
 	GENERATED_BODY()
@@ -36,21 +60,25 @@ protected:
 
 public:
 	/** Whether to pin the Sound Cue asset type when creating new assets. Requires editor restart to take effect. */
-	UPROPERTY(EditAnywhere, config, Category = AssetMenu)
+	UPROPERTY(EditAnywhere, config, Category = AssetMenu, meta = (ConfigRestartRequired = true))
 	bool bPinSoundCueInAssetMenu = true;
 
 	/** Whether to pin the Sound Cue Template asset type when creating new assets. Requires editor restart to take effect. */
-	UPROPERTY(EditAnywhere, config, Category = AssetMenu)
+	UPROPERTY(EditAnywhere, config, Category = AssetMenu, meta = (ConfigRestartRequired = true))
 	bool bPinSoundCueTemplateInAssetMenu = false;
 
 	/** Whether to pin the Sound Attenuation asset type when creating new assets. Requires editor restart to take effect. */
-	UPROPERTY(EditAnywhere, config, Category = AssetMenu)
+	UPROPERTY(EditAnywhere, config, Category = AssetMenu, meta = (ConfigRestartRequired = true))
 	bool bPinSoundAttenuationInAssetMenu = true;
 
 	/** Whether to pin the Sound Concurrency asset type when creating new assets. Requires editor restart to take effect. */
-	UPROPERTY(EditAnywhere, config, Category = AssetMenu)
+	UPROPERTY(EditAnywhere, config, Category = AssetMenu, meta = (ConfigRestartRequired = true))
 	bool bPinSoundConcurrencyInAssetMenu = true;
 
+	/** The device id of the currently selected audio output device. Requires editor restart to take effect */
+	UPROPERTY(EditAnywhere, config, Category = AudioOutputDevice, DisplayName = "Platform Audio API", meta = (ConfigRestartRequired = true))
+	FString AudioMixerModuleName;
+
 	/** Use current audio playback device selected in the operating system. */
 	UPROPERTY(EditAnywhere, config, Category = AudioOutputDevice, DisplayName = "Use System Default Audio Device")
 	bool bUseSystemDevice = true;
@@ -59,6 +87,10 @@ public:
 	UPROPERTY(EditAnywhere, config, Category = AudioOutputDevice)
 	FString AudioOutputDeviceId;
 
+	/** Array of available audio platforms. */
+	UPROPERTY(config)
+	TArray<FAudioPlatform> AudioMixerPlatforms;
+	
 	/** Set and apply, whether audio attenuation is used for non-game worlds*/
 	void SetUseAudioAttenuation(bool bInUseAudioAttenuation);
 
@@ -95,9 +127,10 @@ public:
 	virtual void CustomizeDetails(IDetailLayoutBuilder& DetailBuilder) override;
 	//~ End IDetailCustomization
 
-	static bool IsUIEnabled();
+	static bool IsDeviceSettingsEnabled();
 
 private:
+	TSharedPtr<IPropertyHandle> AudioPlatformPropertyHandle;
 	TSharedPtr<IPropertyHandle> UseSystemDevicePropertyHandle;
 	TSharedPtr<IPropertyHandle> AudioDeviceIdPropertyHandle;
 	TSharedPtr<SWidget> DeviceListComboButton;
@@ -105,6 +138,9 @@ private:
 	TWeakPtr<FActiveTimerHandle> DeviceMenuActiveTimer;
 	FTickerDelegate TickDelegate;
 	FTSTicker::FDelegateHandle TickDelegateHandle;
+	TArray<FAudioPlatform> AudioMixerPlatforms;
+	TArray<TSharedPtr<FString>> AudioBackends;
+	TSharedPtr<FString> CurrentBackendName;
 	
 	static TArray<Audio::FAudioPlatformDeviceInfo> GetAvailableAudioOutputDevices();
 	static Audio::FAudioPlatformDeviceInfo GetDeviceInfo(const FString& InDeviceId);
@@ -115,7 +151,10 @@ private:
 	static void RequestDeviceSwap(const FString& InDeviceId);
 	static bool IsAggregateHardwareDeviceId(const FString& InDeviceId);
 	
+	void AddPlatformRow(IDetailCategoryBuilder& Category, const IDetailLayoutBuilder& DetailLayout);
 	TSharedRef<SWidget> OnGenerateDeviceMenu();
 	void MenuItemDeviceSelected(Audio::FAudioPlatformDeviceInfo InDeviceInfo);
 	bool Tick(float DeltaTime);
+	TSharedPtr<FString> GetCurrentBackendName() const;
+	void SetCurrentBackendName(const FString& InBackendName);
 };
diff --git a/Engine/Source/Editor/AudioEditor/Private/AudioEditorSettings.cpp b/Engine/Source/Editor/AudioEditor/Private/AudioEditorSettings.cpp
--- a/Engine/Source/Editor/AudioEditor/Private/AudioEditorSettings.cpp
+++ b/Engine/Source/Editor/AudioEditor/Private/AudioEditorSettings.cpp
@@ -7,6 +7,8 @@
 #include "DetailsNameWidgetOverrideCustomization.h"
 #include "Editor.h"
 #include "Framework/MultiBox/MultiBoxBuilder.h"
+#include "Misc/ConfigCacheIni.h"
+#include "SSearchableComboBox.h"
 #include "SSimpleComboButton.h"
 #include "Widgets/Input/SCheckBox.h"
 
@@ -15,6 +17,7 @@
 struct FPropertyChangedEvent;
 using namespace Audio;
 
+
 void UAudioEditorSettings::PostInitProperties()
 {
 	Super::PostInitProperties();
@@ -75,6 +78,22 @@ TSharedRef<IDetailCustomization> FAudioOutputDeviceCustomization::MakeInstance()
 FAudioOutputDeviceCustomization::FAudioOutputDeviceCustomization()
 {
 	TickDelegate = FTickerDelegate::CreateRaw(this, &FAudioOutputDeviceCustomization::Tick);
+
+	// Only populate the AudioBackends array if the settings will be visible to the user
+	if (IsDeviceSettingsEnabled())
+	{
+		if (const UAudioEditorSettings* AudioEditorSettings = GetDefault<UAudioEditorSettings>())
+		{
+			AudioMixerPlatforms = AudioEditorSettings->AudioMixerPlatforms;
+			
+			for (const FAudioPlatform& AudioPlatform : AudioMixerPlatforms)
+			{
+				AudioBackends.Add(MakeShared<FString>(AudioPlatform.DisplayName));
+			}
+		}
+	}
+	
+	CurrentBackendName = GetCurrentBackendName();
 }
 
 FAudioOutputDeviceCustomization::~FAudioOutputDeviceCustomization()
@@ -113,14 +132,25 @@ bool FAudioOutputDeviceCustomization::Tick(float DeltaTime)
 void FAudioOutputDeviceCustomization::CustomizeDetails(IDetailLayoutBuilder& DetailBuilder)
 {
 	IDetailCategoryBuilder& Category = DetailBuilder.EditCategory(TEXT("AudioOutputDevice"));
-	IDetailLayoutBuilder& DetailLayout = Category.GetParentLayout();
 
+	// Hide the entire category on platforms where this functionality is not enabled
+	if (!IsDeviceSettingsEnabled())
+	{
+		Category.SetCategoryVisibility(false);
+		return;
+	}
+	
+	const IDetailLayoutBuilder& DetailLayout = Category.GetParentLayout();
+	// Add the Platform Audio API menu on supported platforms
+	AddPlatformRow(Category, DetailLayout);
+	
 	UseSystemDevicePropertyHandle = DetailLayout.GetProperty(GET_MEMBER_NAME_CHECKED(UAudioEditorSettings, bUseSystemDevice));
-	IDetailPropertyRow& UseSystemPropertyRow = Category.AddProperty(UseSystemDevicePropertyHandle);
 
 	// Create check box which indicates system default audio device should be used (disables device combo button)
 	if (UseSystemDevicePropertyHandle.IsValid())
 	{
+		IDetailPropertyRow& UseSystemPropertyRow = Category.AddProperty(UseSystemDevicePropertyHandle);
+		
 		UseSystemPropertyRow.CustomWidget()
 		.NameContent()
 		[
@@ -129,7 +159,6 @@ void FAudioOutputDeviceCustomization::CustomizeDetails(IDetailLayoutBuilder& Det
 		.ValueContent()
 		[
 			SNew(SCheckBox)
-			.IsEnabled(IsUIEnabled())
 			.IsChecked_Lambda([this]() -> ECheckBoxState
 			{
 				bool bUseSystemDefault;
@@ -169,52 +198,56 @@ void FAudioOutputDeviceCustomization::CustomizeDetails(IDetailLayoutBuilder& Det
 	}
 	
 	AudioDeviceIdPropertyHandle = DetailLayout.GetProperty(TEXT("AudioOutputDeviceId"));
-	IDetailPropertyRow& AudioDevicePropertyRow = Category.AddProperty(AudioDeviceIdPropertyHandle);
-
-	bool bUseSystemDefault;
-	UseSystemDevicePropertyHandle->GetValue(bUseSystemDefault);
-
-	// Populate initial value of CurrentDeviceName based depending on the current settings
-	if (bUseSystemDefault)
+	if (AudioDeviceIdPropertyHandle.IsValid())
 	{
-		CurrentDeviceName = GetCurrentAudioMixerDeviceName();
-	}
-	else
-	{
-		FString CurrentDeviceId;
-		AudioDeviceIdPropertyHandle->GetValue(CurrentDeviceId);
-		CurrentDeviceName = GetDeviceNameForDeviceId(CurrentDeviceId);
-	}
+		IDetailPropertyRow& AudioDevicePropertyRow = Category.AddProperty(AudioDeviceIdPropertyHandle);
+
+		bool bUseSystemDefault;
+		UseSystemDevicePropertyHandle->GetValue(bUseSystemDefault);
 
-	AudioDevicePropertyRow.CustomWidget()
-	.NameContent()
-	[
-		SNew(STextBlock)
-		.Text(FText::FromString("Available Audio Output Devices"))
-		.Font(IDetailLayoutBuilder::GetDetailFont())
-	]
-	.ValueContent()
-	[
-		SNew(SSimpleComboButton)
-		.OnGetMenuContent(this, &FAudioOutputDeviceCustomization::OnGenerateDeviceMenu)
-		.ToolTipText(LOCTEXT("AudioOutputDeviceMenuToolTip", "Available Audio Output Devices"))
-		.HasDownArrow(true)
-		.UsesSmallText(true)
-		.IsEnabled_Lambda([this]()
+		// Populate initial value of CurrentDeviceName based depending on the current settings
+		if (bUseSystemDefault)
 		{
-			bool bUseSystemDefault;
-			UseSystemDevicePropertyHandle->GetValue(bUseSystemDefault);
-			return IsUIEnabled() && !bUseSystemDefault;
-		})
-		.Text_Lambda([this]()
+			CurrentDeviceName = GetCurrentAudioMixerDeviceName();
+		}
+		else
 		{
-			return FText::AsCultureInvariant(CurrentDeviceName);
-		})
-	];
+			FString CurrentDeviceId;
+			AudioDeviceIdPropertyHandle->GetValue(CurrentDeviceId);
+			CurrentDeviceName = GetDeviceNameForDeviceId(CurrentDeviceId);
+		}
+
+		AudioDevicePropertyRow.CustomWidget()
+		.NameContent()
+		[
+			SNew(STextBlock)
+			.Text(FText::FromString(TEXT("Available Audio Output Devices")))
+			.ToolTipText(LOCTEXT("AudioOutputDeviceMenuNameToolTip", "Available Audio Output Devices"))
+			.Font(IDetailLayoutBuilder::GetDetailFont())
+		]
+		.ValueContent()
+		[
+			SNew(SSimpleComboButton)
+			.OnGetMenuContent(this, &FAudioOutputDeviceCustomization::OnGenerateDeviceMenu)
+			.ToolTipText(LOCTEXT("AudioOutputDeviceMenuValueToolTip", "Available Audio Output Devices"))
+			.HasDownArrow(true)
+			.UsesSmallText(true)
+			.IsEnabled_Lambda([this]()
+			{
+				bool bUseDefaultValue;
+				UseSystemDevicePropertyHandle->GetValue(bUseDefaultValue);
+				return !bUseDefaultValue;
+			})
+			.Text_Lambda([this]()
+			{
+				return FText::AsCultureInvariant(CurrentDeviceName);
+			})
+		];
 
-	TSharedPtr<SWidget> NameWidget;
-	AudioDevicePropertyRow.GetDefaultWidgets(NameWidget, DeviceListComboButton);
-    			
+		TSharedPtr<SWidget> NameWidget;
+		AudioDevicePropertyRow.GetDefaultWidgets(NameWidget, DeviceListComboButton);
+	}
+	
 	constexpr float TickDelay = 0.1f; // 100ms delay between ticks
 	if (!TickDelegateHandle.IsValid())
 	{
@@ -222,10 +255,53 @@ void FAudioOutputDeviceCustomization::CustomizeDetails(IDetailLayoutBuilder& Det
 	}
 }
 
-bool FAudioOutputDeviceCustomization::IsUIEnabled()
+void FAudioOutputDeviceCustomization::AddPlatformRow(IDetailCategoryBuilder& Category, const IDetailLayoutBuilder& DetailLayout)
+{
+	AudioPlatformPropertyHandle = DetailLayout.GetProperty(GET_MEMBER_NAME_CHECKED(UAudioEditorSettings, AudioMixerModuleName));
+	if (AudioPlatformPropertyHandle.IsValid())
+	{
+		IDetailPropertyRow& AudioPlatformPropertyRow = Category.AddProperty(AudioPlatformPropertyHandle);
+		
+		AudioPlatformPropertyRow.CustomWidget()
+		.NameContent()
+		[
+			SNew(STextBlock)
+			.Text(FText::FromString(TEXT("Platform Audio API")))
+			.ToolTipText(LOCTEXT("AudioPlatformMenuNameToolTip", "Available Audio Platform API's. Requires restart to take effect."))
+			.Font(IDetailLayoutBuilder::GetDetailFont())
+		]
+		.ValueContent()
+		[
+			SNew(SSearchableComboBox)
+			.SearchVisibility(EVisibility::Collapsed)
+			.InitiallySelectedItem(CurrentBackendName)
+			.ToolTipText(LOCTEXT("AudioPlatformMenuValueToolTip", "Available Audio Platform API's. Requires restart to take effect."))
+			.OptionsSource(&AudioBackends)
+			.OnGenerateWidget_Lambda([](TSharedPtr<FString> InItem)
+			{
+				return SNew(STextBlock)
+						.Text(FText::FromString(*InItem))
+						.Font(IDetailLayoutBuilder::GetDetailFont());
+			})
+			.OnSelectionChanged_Lambda([this](TSharedPtr<FString> NewChoice, ESelectInfo::Type SelectType)
+			{
+				SetCurrentBackendName(*NewChoice);
+			})
+			[
+				SNew(STextBlock)
+				.Text_Lambda([this]()
+				{
+					return FText::FromString(*CurrentBackendName);
+				})
+				.Font(IDetailLayoutBuilder::GetDetailFont())
+			]
+		];
+	}
+}
+
+bool FAudioOutputDeviceCustomization::IsDeviceSettingsEnabled()
 {
-	// Only enable UI on Windows platform
-#if PLATFORM_WINDOWS
+#if ENABLE_DEVICE_SETTINGS
 	return true;
 #else
 	return false;
@@ -457,4 +533,46 @@ void FAudioOutputDeviceCustomization::MenuItemDeviceSelected(FAudioPlatformDevic
 	RequestDeviceSwap(InDeviceInfo.DeviceId);
 }
 
+TSharedPtr<FString> FAudioOutputDeviceCustomization::GetCurrentBackendName() const
+{
+	FString DefaultAudioPlatform;
+
+	FString SectionName = UAudioEditorSettings::StaticClass()->GetPathName();
+	// Check to see if the editor pref has been set. If not, fall back to engine setting.
+	if (!GConfig->GetString(*SectionName, TEXT("AudioMixerModuleName"), DefaultAudioPlatform, GEditorSettingsIni))
+	{
+		GConfig->GetString(TEXT("Audio"), TEXT("AudioMixerModuleName"), DefaultAudioPlatform, GEngineIni);
+	}
+
+	int32 Index = AudioMixerPlatforms.IndexOfByPredicate([DefaultAudioPlatform](const FAudioPlatform& Platform)
+	{
+		return Platform.ModuleName == DefaultAudioPlatform;
+	});
+	
+	if (AudioBackends.IsValidIndex(Index))
+	{
+		return AudioBackends[Index];
+	}
+
+	return MakeShared<FString>(TEXT("UNKNOWN"));
+}
+
+void FAudioOutputDeviceCustomization::SetCurrentBackendName(const FString& InBackendName)
+{
+	int32 Index = AudioMixerPlatforms.IndexOfByPredicate([InBackendName](const FAudioPlatform& Platform)
+	{
+		return Platform.DisplayName == InBackendName;
+	});
+	
+	if (AudioMixerPlatforms.IsValidIndex(Index) && AudioPlatformPropertyHandle.IsValid())
+	{
+		AudioPlatformPropertyHandle->SetValue(AudioMixerPlatforms[Index].ModuleName);
+	}
+	
+	if (AudioBackends.IsValidIndex(Index))
+	{
+		CurrentBackendName = AudioBackends[Index];
+	}
+}
+
 #undef LOCTEXT_NAMESPACE
\ No newline at end of file
diff --git a/Engine/Source/Runtime/Engine/Private/AudioDeviceManager.cpp b/Engine/Source/Runtime/Engine/Private/AudioDeviceManager.cpp
--- a/Engine/Source/Runtime/Engine/Private/AudioDeviceManager.cpp
+++ b/Engine/Source/Runtime/Engine/Private/AudioDeviceManager.cpp
@@ -476,9 +476,18 @@ bool FAudioDeviceManager::LoadDefaultAudioDeviceModule()
 	check(!AudioDeviceModule);
 
 	bool bForceNonRealtimeRenderer = FParse::Param(FCommandLine::Get(), TEXT("DeterministicAudio"));
+	bool bFoundModuleName = false;
+	
+#if WITH_EDITOR
+	// Check to see if the editor pref has been set. If not, fall back to engine setting.
+	bFoundModuleName = GConfig->GetString(TEXT("/Script/AudioEditor.AudioEditorSettings"), TEXT("AudioMixerModuleName"), AudioMixerModuleName, GEditorSettingsIni);
+#endif //WITH_EDITOR
 
-	// If not using command line switch to use audio mixer, check the game platform engine ini file (e.g. WindowsEngine.ini) which enables it for player
-	GConfig->GetString(TEXT("Audio"), TEXT("AudioMixerModuleName"), AudioMixerModuleName, GEngineIni);
+	if (!bFoundModuleName)
+	{
+		// If not using command line switch to use audio mixer, check the game platform engine ini file (e.g. WindowsEngine.ini) which enables it for player
+		GConfig->GetString(TEXT("Audio"), TEXT("AudioMixerModuleName"), AudioMixerModuleName, GEngineIni);
+	}
 
 	if (bForceNonRealtimeRenderer)
 	{
