Change macro ENABLE_DEVICE_SETTINGS to be more specific to avoid any potential collisions (ENABLE_AUDIO_DEVICE_EDITOR_SETTINGS)

#rnx

diff --git a/Engine/Source/Editor/AudioEditor/AudioEditor.Build.cs b/Engine/Source/Editor/AudioEditor/AudioEditor.Build.cs
--- a/Engine/Source/Editor/AudioEditor/AudioEditor.Build.cs
+++ b/Engine/Source/Editor/AudioEditor/AudioEditor.Build.cs
@@ -86,7 +86,7 @@ public class AudioEditor : ModuleRules
 			RuntimeDependencies.Add("$(EngineDir)/Binaries/ThirdParty/libsndfile/" + PlatformName + "/libsndfile-1.dll");
 
 			PublicDefinitions.Add("WITH_SNDFILE_IO=1");
-			PublicDefinitions.Add("ENABLE_DEVICE_SETTINGS=1");
+			PublicDefinitions.Add("ENABLE_AUDIO_DEVICE_EDITOR_SETTINGS=1");
 		}
 		else if (Target.IsInPlatformGroup(UnrealPlatformGroup.Apple))
 		{
@@ -100,7 +100,7 @@ public class AudioEditor : ModuleRules
 			RuntimeDependencies.Add(LibSndFilePath);
 			
 			PublicDefinitions.Add("WITH_SNDFILE_IO=1");
-			PublicDefinitions.Add("ENABLE_DEVICE_SETTINGS=0");
+			PublicDefinitions.Add("ENABLE_AUDIO_DEVICE_EDITOR_SETTINGS=0");
 		}
 		else if (Target.IsInPlatformGroup(UnrealPlatformGroup.Linux))
 		{
@@ -114,12 +114,12 @@ public class AudioEditor : ModuleRules
 			RuntimeDependencies.Add(LibSndFilePath);
 			
 			PublicDefinitions.Add("WITH_SNDFILE_IO=1");
-			PublicDefinitions.Add("ENABLE_DEVICE_SETTINGS=0");
+			PublicDefinitions.Add("ENABLE_AUDIO_DEVICE_EDITOR_SETTINGS=0");
 		}
 		else
 		{
 			PublicDefinitions.Add("WITH_SNDFILE_IO=0");
-			PublicDefinitions.Add("ENABLE_DEVICE_SETTINGS=0");
+			PublicDefinitions.Add("ENABLE_AUDIO_DEVICE_EDITOR_SETTINGS=0");
 		}
 	}
 }
diff --git a/Engine/Source/Editor/AudioEditor/Private/AudioEditorSettings.cpp b/Engine/Source/Editor/AudioEditor/Private/AudioEditorSettings.cpp
--- a/Engine/Source/Editor/AudioEditor/Private/AudioEditorSettings.cpp
+++ b/Engine/Source/Editor/AudioEditor/Private/AudioEditorSettings.cpp
@@ -301,7 +301,7 @@ void FAudioOutputDeviceCustomization::AddPlatformRow(IDetailCategoryBuilder& Cat
 
 bool FAudioOutputDeviceCustomization::IsDeviceSettingsEnabled()
 {
-#if ENABLE_DEVICE_SETTINGS
+#if ENABLE_AUDIO_DEVICE_EDITOR_SETTINGS
 	return true;
 #else
 	return false;
