Fix SoundCue generation crash with more than 32 Sound Sources

diff --git a/Engine/Source/Editor/AudioEditor/Private/Factories/SoundCueFactoryNew.cpp b/Engine/Source/Editor/AudioEditor/Private/Factories/SoundCueFactoryNew.cpp
--- a/Engine/Source/Editor/AudioEditor/Private/Factories/SoundCueFactoryNew.cpp
+++ b/Engine/Source/Editor/AudioEditor/Private/Factories/SoundCueFactoryNew.cpp
@@ -2,6 +2,7 @@
 
 #include "Factories/SoundCueFactoryNew.h"
 
+#include "Audio.h"
 #include "EdGraph/EdGraph.h"
 #include "EdGraph/EdGraphNode.h"
 #include "HAL/PlatformCrt.h"
@@ -65,15 +66,23 @@ UObject* USoundCueFactoryNew::FactoryCreateNew(UClass* Class, UObject* InParent,
 			{
 				if (USoundNode* PlayerNode = USoundCueFactoryNew::CreateSoundPlayerNode(SoundCue, SoundObject.Get(), PosX, PosY))
 				{
-					const int32 ChildNodeIndex = RandomNode->ChildNodes.Num();
-					RandomNode->InsertChildNode(ChildNodeIndex);
-					RandomNode->ChildNodes[ChildNodeIndex] = PlayerNode;
+					const int32 ChildNodeIndex = RandomNode->ChildNodes.Num(); 
+					if (ChildNodeIndex < RandomNode->GetMaxChildNodes())
+					{ 
+						RandomNode->InsertChildNode(ChildNodeIndex);
+						RandomNode->ChildNodes[ChildNodeIndex] = PlayerNode;
+					}
 
 					PlayerNodes.Add(PlayerNode);
 					PosY += 100;
 				}
 			}
 
+			if (SoundObjects.Num() > RandomNode->GetMaxChildNodes())
+			{
+				UE_LOG(LogAudio, Warning, TEXT("RandomNode is limited to %d input Sound Sources. Excess Sound Sources in SoundCue %s will be left unconnected."), RandomNode->GetMaxChildNodes(), *Name.ToString());
+			}
+
 			if (UEdGraph* SoundCueGraph = SoundCue->GetGraph())
 			{
 				if (const USoundCueGraphSchema* GraphSchema = Cast<USoundCueGraphSchema>(SoundCueGraph->GetSchema()))
