Prevent custom soundwave editors from opening when restricted mode is true (fall back to default editor)

#rnx

diff --git a/Engine/Source/Editor/AudioEditor/Private/AssetTypeActions/AssetDefinition_SoundWave.cpp b/Engine/Source/Editor/AudioEditor/Private/AssetTypeActions/AssetDefinition_SoundWave.cpp
--- a/Engine/Source/Editor/AudioEditor/Private/AssetTypeActions/AssetDefinition_SoundWave.cpp
+++ b/Engine/Source/Editor/AudioEditor/Private/AssetTypeActions/AssetDefinition_SoundWave.cpp
@@ -48,8 +48,9 @@ EAssetCommandResult UAssetDefinition_SoundWave::OpenAssets(const FAssetOpenArgs&
 {
 	IAudioEditorModule& AudioEditorModule = FModuleManager::LoadModuleChecked<IAudioEditorModule>( "AudioEditor" );
 
-	// Use a custom editor for soundwaves if one is registered
-	if (AudioEditorModule.SoundWaveEditorOpen.IsBound())
+	// Use a custom editor for soundwaves if one is registered except in restricted mode
+	// which currently does not support custom soundwave editors
+	if (AudioEditorModule.SoundWaveEditorOpen.IsBound() && !AudioEditorModule.IsRestrictedMode())
 	{
 		AudioEditorModule.SoundWaveEditorOpen.Execute(OpenArgs.LoadObjects<USoundWave>());
 	}
