Fix up naming FText collisions
- Minor clean-up for distinguishing action names

#rnx

diff --git a/Engine/Source/Editor/AudioEditor/Private/AssetTypeActions/AssetDefinition_SoundBase.cpp b/Engine/Source/Editor/AudioEditor/Private/AssetTypeActions/AssetDefinition_SoundBase.cpp
--- a/Engine/Source/Editor/AudioEditor/Private/AssetTypeActions/AssetDefinition_SoundBase.cpp
+++ b/Engine/Source/Editor/AudioEditor/Private/AssetTypeActions/AssetDefinition_SoundBase.cpp
@@ -15,7 +15,7 @@
 #include "Widgets/Input/SButton.h"
 #include "Widgets/Images/SImage.h"
 
-#define LOCTEXT_NAMESPACE "AssetTypeActions"
+#define LOCTEXT_NAMESPACE "AudioEditor"
 
 namespace UE::AudioEditor
 {
@@ -699,4 +699,4 @@ namespace MenuExtension_SoundBase
 	});
 }
 
-#undef LOCTEXT_NAMESPACE
+#undef LOCTEXT_NAMESPACE // AudioEditor
diff --git a/Engine/Source/Editor/AudioEditor/Private/AssetTypeActions/AssetTypeActions_DialogueWave.cpp b/Engine/Source/Editor/AudioEditor/Private/AssetTypeActions/AssetTypeActions_DialogueWave.cpp
--- a/Engine/Source/Editor/AudioEditor/Private/AssetTypeActions/AssetTypeActions_DialogueWave.cpp
+++ b/Engine/Source/Editor/AudioEditor/Private/AssetTypeActions/AssetTypeActions_DialogueWave.cpp
@@ -31,16 +31,16 @@
 class IToolkitHost;
 class USoundBase;
 
-#define LOCTEXT_NAMESPACE "AssetTypeActions"
+#define LOCTEXT_NAMESPACE "AudioEditor"
 
 void FAssetTypeActions_DialogueWave::GetActions(const TArray<UObject*>& InObjects, FToolMenuSection& Section)
 {
 	auto DialogueWaves = GetTypedWeakObjectPtrs<UDialogueWave>(InObjects);
 
 	Section.AddMenuEntry(
-		"Sound_PlaySound",
-		LOCTEXT("Sound_PlaySound", "Play"),
-		LOCTEXT("Sound_PlaySoundTooltip", "Plays the selected sound."),
+		"DialogWave_PlaySound",
+		LOCTEXT("DialogWave_PlaySound", "Play"),
+		LOCTEXT("DialogWave_PlaySoundTooltip", "Plays the selected dialog wave."),
 		FSlateIcon(FAppStyle::GetAppStyleSetName(), "MediaAsset.AssetActions.Play.Small"),
 		FUIAction(
 			FExecuteAction::CreateSP(this, &FAssetTypeActions_DialogueWave::ExecutePlaySound, DialogueWaves),
@@ -49,9 +49,9 @@ void FAssetTypeActions_DialogueWave::GetActions(const TArray<UObject*>& InObject
 	);
 
 	Section.AddMenuEntry(
-		"Sound_StopSound",
-		LOCTEXT("Sound_StopSound", "Stop"),
-		LOCTEXT("Sound_StopSoundTooltip", "Stops the selected sounds."),
+		"DialogWave_StopSound",
+		LOCTEXT("DialogWave_Stop", "Stop"),
+		LOCTEXT("DialogWave_StopTooltip", "Stops the selected dialog wave(s)."),
 		FSlateIcon(FAppStyle::GetAppStyleSetName(), "MediaAsset.AssetActions.Stop.Small"),
 		FUIAction(
 			FExecuteAction::CreateSP(this, &FAssetTypeActions_DialogueWave::ExecuteStopSound, DialogueWaves),
@@ -256,4 +256,4 @@ const TArray<FText>& FAssetTypeActions_DialogueWave::GetSubMenus() const
 
 	return SubMenus;
 }
-#undef LOCTEXT_NAMESPACE
+#undef LOCTEXT_NAMESPACE // AudioEditor
