ATR - Remove the "Copy settings from selected SoundWave" prompt on import.

[RN] Use of a Template SoundWave has been disabled and can be re-enable in Editor Preferences->Audio Editor Per Project Settings

diff --git a/Engine/Source/Editor/AudioEditor/Classes/AudioEditorPerProjectUserSettings.h b/Engine/Source/Editor/AudioEditor/Classes/AudioEditorPerProjectUserSettings.h
new file mode 100644
--- /dev/null
+++ b/Engine/Source/Editor/AudioEditor/Classes/AudioEditorPerProjectUserSettings.h
@@ -0,0 +1,31 @@
+// Copyright Epic Games, Inc. All Rights Reserved.
+
+#pragma once
+
+#include "Engine/DeveloperSettings.h"
+
+#include "AudioEditorPerProjectUserSettings.generated.h"
+
+UENUM()
+enum class EUseTemplateSoundWaveDuringAssetImport : uint8
+{
+	AlwaysPrompt = 0,
+	AlwaysUse = 1,
+	NeverUse = 2
+};
+
+UCLASS(MinimalAPI, config = EditorPerProjectUserSettings, meta = (DisplayName = "Audio Editor Per Project Settings"))
+class UAudioEditorPerProjectUserSettings : public UDeveloperSettings
+{
+	GENERATED_BODY()
+
+public:
+	/** Whether to use the selected Sound Wave in the Content Browser as a template for sound(s) being imported. */
+	UPROPERTY(EditAnywhere, config, Category = AssetImport)
+	EUseTemplateSoundWaveDuringAssetImport UseTemplateSoundWaveDuringAssetImport = EUseTemplateSoundWaveDuringAssetImport::NeverUse;
+
+	//~ Begin UDeveloperSettings
+	virtual FName GetCategoryName() const override { return TEXT("General"); }
+	//~ End UDeveloperSettings
+
+};
diff --git a/Engine/Source/Editor/AudioEditor/Private/Factories/SoundFactory.cpp b/Engine/Source/Editor/AudioEditor/Private/Factories/SoundFactory.cpp
--- a/Engine/Source/Editor/AudioEditor/Private/Factories/SoundFactory.cpp
+++ b/Engine/Source/Editor/AudioEditor/Private/Factories/SoundFactory.cpp
@@ -4,6 +4,7 @@
 #include "AssetRegistry/AssetRegistryModule.h"
 #include "Audio.h"
 #include "AudioAnalytics.h"
+#include "AudioEditorPerProjectUserSettings.h"
 #include "Components/AudioComponent.h"
 #include "ContentBrowserModule.h"
 #include "IContentBrowserSingleton.h"
@@ -27,6 +28,7 @@
 #include "AssetToolsModule.h"
 #include "IAssetTools.h"
 #include "Math/NumericLimits.h"
+#include "Widgets/Layout/SUniformGridPanel.h"
 
 // Disable user import
 
@@ -792,6 +794,89 @@ UObject* USoundFactory::CreateObject
 	return nullptr;
 }
 
+bool ShowUseTemplateDialog(const FString& SoundWaveName, bool& bOutRememberChoice)
+{
+	bool bUserChoice = false;
+	bOutRememberChoice = false;
+
+	TSharedPtr<SWindow> DialogWindow = SNew(SWindow)
+		.Title(FText::FromString("SoundWave Import Options"))
+		.SizingRule(ESizingRule::Autosized)
+		.AutoCenter(EAutoCenter::PreferredWorkArea)
+		.SupportsMinimize(false).SupportsMaximize(false);
+
+	TWeakPtr<SWindow> WeakWindow = DialogWindow;
+
+	TSharedPtr<SCheckBox> RememberCheckBox;
+
+	DialogWindow->SetContent(
+		SNew(SBorder)
+		.Padding(10)
+		[
+			SNew(SVerticalBox)
+
+				+ SVerticalBox::Slot().AutoHeight().Padding(5)
+				[
+					SNew(STextBlock)
+						.AutoWrapText(true)
+						.Text(FText::Format(
+							NSLOCTEXT("SoundFactory", "UseSoundWaveTemplate", "Use the selected Sound Wave '{0}' as a template for sound(s) being imported?"),
+							FText::FromString(SoundWaveName)))
+				]
+
+			+ SVerticalBox::Slot().AutoHeight().Padding(5)
+				[
+					SAssignNew(RememberCheckBox, SCheckBox)
+						[
+							SNew(STextBlock).Text(NSLOCTEXT("SoundFactory", "RememberChoice", "Remember my choice"))
+						]
+				]
+
+				+ SVerticalBox::Slot().AutoHeight().HAlign(HAlign_Right).Padding(5)
+				[
+					SNew(SUniformGridPanel).SlotPadding(5)
+
+						+ SUniformGridPanel::Slot(0, 0)
+						[
+							SNew(SButton)
+								.Text(NSLOCTEXT("Generic", "Yes", "Yes"))
+								.OnClicked_Lambda([WeakWindow, &bUserChoice]() 
+									{
+										bUserChoice = true;
+
+										if (TSharedPtr<SWindow> PinnedWindow = WeakWindow.Pin())
+										{
+											FSlateApplication::Get().RequestDestroyWindow(PinnedWindow.ToSharedRef());
+										}
+										
+										return FReply::Handled();
+									})
+						]
+
+					+ SUniformGridPanel::Slot(1, 0)
+						[
+							SNew(SButton)
+								.Text(NSLOCTEXT("Generic", "No", "No"))
+								.OnClicked_Lambda([WeakWindow]()
+									{
+										if (TSharedPtr<SWindow> PinnedWindow = WeakWindow.Pin())
+										{
+											FSlateApplication::Get().RequestDestroyWindow(PinnedWindow.ToSharedRef());
+										}
+										
+										return FReply::Handled();
+									})
+						]
+				]
+		]);
+
+	FSlateApplication::Get().AddModalWindow(DialogWindow.ToSharedRef(), nullptr);
+
+	bOutRememberChoice = RememberCheckBox.IsValid() && RememberCheckBox->IsChecked();
+
+	return bUserChoice;
+}
+
 void USoundFactory::SuppressImportDialogs()
 {
 	SuppressImportDialogOptions = ESuppressImportDialog::Overwrite | ESuppressImportDialog::UseTemplate;
@@ -799,7 +884,10 @@ void USoundFactory::SuppressImportDialogs()
 
 void USoundFactory::UpdateTemplate()
 {
-	if (!IsAutomatedImport() && !TemplateSoundWave.IsValid() && !(SuppressImportDialogOptions & ESuppressImportDialog::UseTemplate))
+	EUseTemplateSoundWaveDuringAssetImport& UseTemplateSoundWaveEditorSetting = GetMutableDefault<UAudioEditorPerProjectUserSettings>()->UseTemplateSoundWaveDuringAssetImport;
+	const bool bUseTemplate = UseTemplateSoundWaveEditorSetting != EUseTemplateSoundWaveDuringAssetImport::NeverUse;
+
+	if (!IsAutomatedImport() && bUseTemplate && !TemplateSoundWave.IsValid() && !(SuppressImportDialogOptions & ESuppressImportDialog::UseTemplate))
 	{
 		SuppressImportDialogOptions |= ESuppressImportDialog::UseTemplate;
 
@@ -807,18 +895,43 @@ void USoundFactory::UpdateTemplate()
 		TArray<FAssetData> SelectedAssets;
 		ContentBrowserModule.Get().GetSelectedAssets(SelectedAssets);
 
+		// If SelectedAssets.Num() != 1, no TemplateSoundWave is used
 		if (SelectedAssets.Num() == 1)
 		{
 			if (USoundWave* SoundWave = Cast<USoundWave>(SelectedAssets[0].GetAsset()))
 			{
-				const bool bUseTemplateSoundWave = FMessageDialog::Open(EAppMsgType::YesNo, FText::Format(
-					NSLOCTEXT("SoundFactory", "UseSoundWaveTemplate", "Use the selected Sound Wave '{0}' in the Content Browser as a template for sound(s) being imported?"),
-					FText::FromString(SoundWave->GetName()))) == EAppReturnType::Yes;
+				if (UseTemplateSoundWaveEditorSetting == EUseTemplateSoundWaveDuringAssetImport::AlwaysUse)
+				{
+					TemplateSoundWave = SoundWave;
+
+					return;
+				}
+
+				check(UseTemplateSoundWaveEditorSetting == EUseTemplateSoundWaveDuringAssetImport::AlwaysPrompt)
+
+				bool bRememberSelectedChoice = false; 
+				const bool bUseTemplateSoundWave = ShowUseTemplateDialog(SoundWave->GetName(), bRememberSelectedChoice);
 
 				if (bUseTemplateSoundWave)
 				{
 					TemplateSoundWave = SoundWave;
 				}
+
+				if (bRememberSelectedChoice)
+				{
+					if (bUseTemplateSoundWave)
+					{
+						UseTemplateSoundWaveEditorSetting = EUseTemplateSoundWaveDuringAssetImport::AlwaysUse;
+					}
+					else
+					{
+						UseTemplateSoundWaveEditorSetting = EUseTemplateSoundWaveDuringAssetImport::NeverUse;
+					}
+				}
+				else
+				{
+					UseTemplateSoundWaveEditorSetting = EUseTemplateSoundWaveDuringAssetImport::AlwaysPrompt;
+				}
 			}
 		}
 	}
