Implemented USoundClassGraphSchema::GetAssetsGraphHoverMessage to provide asset drag and drop feedback.

diff --git a/Engine/Source/Editor/AudioEditor/Classes/SoundClassGraph/SoundClassGraphSchema.h b/Engine/Source/Editor/AudioEditor/Classes/SoundClassGraph/SoundClassGraphSchema.h
--- a/Engine/Source/Editor/AudioEditor/Classes/SoundClassGraph/SoundClassGraphSchema.h
+++ b/Engine/Source/Editor/AudioEditor/Classes/SoundClassGraph/SoundClassGraphSchema.h
@@ -74,6 +74,7 @@ class USoundClassGraphSchema : public UEdGraphSchema
 	virtual void BreakPinLinks(UEdGraphPin& TargetPin, bool bSendsNodeNotifcation) const override;
 	virtual void BreakSinglePinLink(UEdGraphPin* SourcePin, UEdGraphPin* TargetPin) const override;
 	virtual void DroppedAssetsOnGraph(const TArray<struct FAssetData>& Assets, const FVector2f& GraphPosition, UEdGraph* Graph) const override;
+	virtual void GetAssetsGraphHoverMessage(const TArray<FAssetData>& Assets, const UEdGraph* HoverGraph, FString& OutTooltipText, bool& OutOkIcon) const override;
 	//~ End EdGraphSchema Interface
 };
 
diff --git a/Engine/Source/Editor/AudioEditor/Private/SoundClassGraphSchema.cpp b/Engine/Source/Editor/AudioEditor/Private/SoundClassGraphSchema.cpp
--- a/Engine/Source/Editor/AudioEditor/Private/SoundClassGraphSchema.cpp
+++ b/Engine/Source/Editor/AudioEditor/Private/SoundClassGraphSchema.cpp
@@ -182,4 +182,29 @@ void USoundClassGraphSchema::DroppedAssetsOnGraph(const TArray<struct FAssetData
 	}
 }
 
+void USoundClassGraphSchema::GetAssetsGraphHoverMessage(const TArray<FAssetData>& Assets, const UEdGraph* HoverGraph, FString& OutTooltipText, bool& OutOkIcon) const
+{
+	OutTooltipText.Reset();
+	OutOkIcon = false;
+
+	const USoundClassGraph* SoundClassGraph = CastChecked<const USoundClassGraph>(HoverGraph);
+	for (int32 AssetIdx = 0; AssetIdx < Assets.Num(); ++AssetIdx)
+	{
+		USoundClass* SoundClass = Cast<USoundClass>(Assets[AssetIdx].GetAsset());
+		if (SoundClass)
+		{
+			if (!SoundClassGraph->IsClassDisplayed(SoundClass))
+			{
+				OutTooltipText = FText::FormatOrdered(LOCTEXT("DropSoundClass", "Add {0} to Graph."), USoundClass::StaticClass()->GetDisplayNameText()).ToString();
+				OutOkIcon = true;
+				break;
+			}
+			else if (OutTooltipText.IsEmpty())
+			{
+				OutTooltipText = FText::FormatOrdered(LOCTEXT("SoundClassInGraph", "{0} already in graph."), USoundClass::StaticClass()->GetDisplayNameText()).ToString();
+			}
+		}
+	}
+}
+
 #undef LOCTEXT_NAMESPACE
