[Audio] Implement standard diffing for audio engine assets

#tests Diffed couples of assets for each type I've added standard diffing with, all working

diff --git a/Engine/Plugins/Runtime/AudioModulation/Source/AudioModulationEditor/Private/AssetTypeActions/AssetTypeActions_SoundControlBus.h b/Engine/Plugins/Runtime/AudioModulation/Source/AudioModulationEditor/Private/AssetTypeActions/AssetTypeActions_SoundControlBus.h
--- a/Engine/Plugins/Runtime/AudioModulation/Source/AudioModulationEditor/Private/AssetTypeActions/AssetTypeActions_SoundControlBus.h
+++ b/Engine/Plugins/Runtime/AudioModulation/Source/AudioModulationEditor/Private/AssetTypeActions/AssetTypeActions_SoundControlBus.h
@@ -1,12 +1,12 @@
 // Copyright Epic Games, Inc. All Rights Reserved.
 #pragma once
 
-#include "CoreMinimal.h"
-#include "AssetTypeActions_Base.h"
+#include "AssetTypeActions/AssetTypeActions_AudioDiffable.h"
 #include "AudioModulationStyle.h"
+#include "CoreMinimal.h"
 
 
-class FAssetTypeActions_SoundControlBus : public FAssetTypeActions_Base
+class FAssetTypeActions_SoundControlBus : public FAssetTypeActions_AudioDiffable
 {
 public:
 	// IAssetTypeActions Implementation
diff --git a/Engine/Plugins/Runtime/AudioModulation/Source/AudioModulationEditor/Private/AssetTypeActions/AssetTypeActions_SoundControlBusMix.h b/Engine/Plugins/Runtime/AudioModulation/Source/AudioModulationEditor/Private/AssetTypeActions/AssetTypeActions_SoundControlBusMix.h
--- a/Engine/Plugins/Runtime/AudioModulation/Source/AudioModulationEditor/Private/AssetTypeActions/AssetTypeActions_SoundControlBusMix.h
+++ b/Engine/Plugins/Runtime/AudioModulation/Source/AudioModulationEditor/Private/AssetTypeActions/AssetTypeActions_SoundControlBusMix.h
@@ -2,11 +2,11 @@
 
 #pragma once
 
-#include "CoreMinimal.h"
-#include "AssetTypeActions_Base.h"
+#include "AssetTypeActions/AssetTypeActions_AudioDiffable.h"
 #include "AudioModulationStyle.h"
+#include "CoreMinimal.h"
 
-class FAssetTypeActions_SoundControlBusMix : public FAssetTypeActions_Base
+class FAssetTypeActions_SoundControlBusMix : public FAssetTypeActions_AudioDiffable
 {
 public:
 	// IAssetTypeActions Implementation
diff --git a/Engine/Plugins/Runtime/AudioModulation/Source/AudioModulationEditor/Private/AssetTypeActions/AssetTypeActions_SoundModulationGenerator.h b/Engine/Plugins/Runtime/AudioModulation/Source/AudioModulationEditor/Private/AssetTypeActions/AssetTypeActions_SoundModulationGenerator.h
--- a/Engine/Plugins/Runtime/AudioModulation/Source/AudioModulationEditor/Private/AssetTypeActions/AssetTypeActions_SoundModulationGenerator.h
+++ b/Engine/Plugins/Runtime/AudioModulation/Source/AudioModulationEditor/Private/AssetTypeActions/AssetTypeActions_SoundModulationGenerator.h
@@ -1,11 +1,11 @@
 // Copyright Epic Games, Inc. All Rights Reserved.
 #pragma once
 
-#include "AssetTypeActions_Base.h"
+#include "AssetTypeActions/AssetTypeActions_AudioDiffable.h"
 #include "AudioModulationStyle.h"
 #include "CoreMinimal.h"
 
-class FAssetTypeActions_SoundModulationGenerator : public FAssetTypeActions_Base
+class FAssetTypeActions_SoundModulationGenerator : public FAssetTypeActions_AudioDiffable
 {
 public:
 	// IAssetTypeActions Implementation
diff --git a/Engine/Plugins/Runtime/AudioModulation/Source/AudioModulationEditor/Private/AssetTypeActions/AssetTypeActions_SoundModulationParameter.h b/Engine/Plugins/Runtime/AudioModulation/Source/AudioModulationEditor/Private/AssetTypeActions/AssetTypeActions_SoundModulationParameter.h
--- a/Engine/Plugins/Runtime/AudioModulation/Source/AudioModulationEditor/Private/AssetTypeActions/AssetTypeActions_SoundModulationParameter.h
+++ b/Engine/Plugins/Runtime/AudioModulation/Source/AudioModulationEditor/Private/AssetTypeActions/AssetTypeActions_SoundModulationParameter.h
@@ -1,12 +1,12 @@
 // Copyright Epic Games, Inc. All Rights Reserved.
 #pragma once
 
-#include "CoreMinimal.h"
-#include "AssetTypeActions_Base.h"
+#include "AssetTypeActions/AssetTypeActions_AudioDiffable.h"
 #include "AudioModulationStyle.h"
+#include "CoreMinimal.h"
 
 
-class FAssetTypeActions_SoundModulationParameter : public FAssetTypeActions_Base
+class FAssetTypeActions_SoundModulationParameter : public FAssetTypeActions_AudioDiffable
 {
 public:
 	// IAssetTypeActions Implementation
diff --git a/Engine/Plugins/Runtime/AudioModulation/Source/AudioModulationEditor/Private/AssetTypeActions/AssetTypeActions_SoundModulationPatch.h b/Engine/Plugins/Runtime/AudioModulation/Source/AudioModulationEditor/Private/AssetTypeActions/AssetTypeActions_SoundModulationPatch.h
--- a/Engine/Plugins/Runtime/AudioModulation/Source/AudioModulationEditor/Private/AssetTypeActions/AssetTypeActions_SoundModulationPatch.h
+++ b/Engine/Plugins/Runtime/AudioModulation/Source/AudioModulationEditor/Private/AssetTypeActions/AssetTypeActions_SoundModulationPatch.h
@@ -2,11 +2,11 @@
 
 #pragma once
 
-#include "CoreMinimal.h"
-#include "AssetTypeActions_Base.h"
+#include "AssetTypeActions/AssetTypeActions_AudioDiffable.h"
 #include "AudioModulationStyle.h"
+#include "CoreMinimal.h"
 
-class FAssetTypeActions_SoundModulationPatch : public FAssetTypeActions_Base
+class FAssetTypeActions_SoundModulationPatch : public FAssetTypeActions_AudioDiffable
 {
 public:
 	// IAssetTypeActions Implementation
diff --git a/Engine/Plugins/Runtime/AudioSynesthesia/Source/AudioSynesthesiaEditor/AudioSynesthesiaEditor.Build.cs b/Engine/Plugins/Runtime/AudioSynesthesia/Source/AudioSynesthesiaEditor/AudioSynesthesiaEditor.Build.cs
--- a/Engine/Plugins/Runtime/AudioSynesthesia/Source/AudioSynesthesiaEditor/AudioSynesthesiaEditor.Build.cs
+++ b/Engine/Plugins/Runtime/AudioSynesthesia/Source/AudioSynesthesiaEditor/AudioSynesthesiaEditor.Build.cs
@@ -6,29 +6,27 @@ public class AudioSynesthesiaEditor : ModuleRules
 {
 	public AudioSynesthesiaEditor(ReadOnlyTargetRules Target) : base(Target)
 	{
-
-		//bFasterWithoutUnity = true;
-
 		PublicDependencyModuleNames.AddRange(
 			new string[]
 			{
+				"AudioAnalyzer",
+				"AudioEditor",
+				"AudioSynesthesia",
 				"Core",
 				"CoreUObject",
 				"Engine",
 				"EditorFramework",
+				"EditorStyle",
+				"InputCore",
+				"Slate",
+				"SlateCore",
 				"UnrealEd",
-				"AudioSynesthesia",
-				"AudioAnalyzer",
-                "InputCore",
-                "Slate",
-                "SlateCore",
-                "EditorStyle"
             }
         );
 
 		PrivateIncludePathModuleNames.AddRange(
 			new string[] {
-				"AssetTools"
+				"AssetTools",
 			});
 	}
 }
diff --git a/Engine/Plugins/Runtime/AudioSynesthesia/Source/AudioSynesthesiaEditor/Private/AssetTypeActions_AudioSynesthesiaNRT.h b/Engine/Plugins/Runtime/AudioSynesthesia/Source/AudioSynesthesiaEditor/Private/AssetTypeActions_AudioSynesthesiaNRT.h
--- a/Engine/Plugins/Runtime/AudioSynesthesia/Source/AudioSynesthesiaEditor/Private/AssetTypeActions_AudioSynesthesiaNRT.h
+++ b/Engine/Plugins/Runtime/AudioSynesthesia/Source/AudioSynesthesiaEditor/Private/AssetTypeActions_AudioSynesthesiaNRT.h
@@ -2,12 +2,12 @@
 
 #pragma once
 
-#include "AssetTypeActions_Base.h"
+#include "AssetTypeActions/AssetTypeActions_AudioDiffable.h"
 
 class UAudioSynesthesiaNRT;
 
 /** FAssetTypeActions_AudioSynesthesiaNRT */
-class FAssetTypeActions_AudioSynesthesiaNRT : public FAssetTypeActions_Base
+class FAssetTypeActions_AudioSynesthesiaNRT : public FAssetTypeActions_AudioDiffable
 {
 public:
 	FAssetTypeActions_AudioSynesthesiaNRT(UAudioSynesthesiaNRT* InSynesthesia);
diff --git a/Engine/Plugins/Runtime/AudioSynesthesia/Source/AudioSynesthesiaEditor/Private/AssetTypeActions_AudioSynesthesiaNRTSettings.h b/Engine/Plugins/Runtime/AudioSynesthesia/Source/AudioSynesthesiaEditor/Private/AssetTypeActions_AudioSynesthesiaNRTSettings.h
--- a/Engine/Plugins/Runtime/AudioSynesthesia/Source/AudioSynesthesiaEditor/Private/AssetTypeActions_AudioSynesthesiaNRTSettings.h
+++ b/Engine/Plugins/Runtime/AudioSynesthesia/Source/AudioSynesthesiaEditor/Private/AssetTypeActions_AudioSynesthesiaNRTSettings.h
@@ -2,12 +2,12 @@
 
 #pragma once
 
-#include "AssetTypeActions_Base.h"
+#include "AssetTypeActions/AssetTypeActions_AudioDiffable.h"
 
 class UAudioSynesthesiaNRTSettings;
 
 /** FAssetTypeActions_AudioSynesthesiaNRTSettings */
-class FAssetTypeActions_AudioSynesthesiaNRTSettings : public FAssetTypeActions_Base
+class FAssetTypeActions_AudioSynesthesiaNRTSettings : public FAssetTypeActions_AudioDiffable
 {
 public:
 	FAssetTypeActions_AudioSynesthesiaNRTSettings(UAudioSynesthesiaNRTSettings* InSynesthesiaSettings);
diff --git a/Engine/Plugins/Runtime/AudioSynesthesia/Source/AudioSynesthesiaEditor/Private/AssetTypeActions_AudioSynesthesiaSettings.h b/Engine/Plugins/Runtime/AudioSynesthesia/Source/AudioSynesthesiaEditor/Private/AssetTypeActions_AudioSynesthesiaSettings.h
--- a/Engine/Plugins/Runtime/AudioSynesthesia/Source/AudioSynesthesiaEditor/Private/AssetTypeActions_AudioSynesthesiaSettings.h
+++ b/Engine/Plugins/Runtime/AudioSynesthesia/Source/AudioSynesthesiaEditor/Private/AssetTypeActions_AudioSynesthesiaSettings.h
@@ -2,12 +2,12 @@
 
 #pragma once
 
-#include "AssetTypeActions_Base.h"
+#include "AssetTypeActions/AssetTypeActions_AudioDiffable.h"
 
 class UAudioSynesthesiaSettings;
 
 /** FAssetTypeActions_AudioSynesthesiaSettings */
-class FAssetTypeActions_AudioSynesthesiaSettings : public FAssetTypeActions_Base
+class FAssetTypeActions_AudioSynesthesiaSettings : public FAssetTypeActions_AudioDiffable
 {
 public:
 	FAssetTypeActions_AudioSynesthesiaSettings(UAudioSynesthesiaSettings* InSynesthesiaSettings);
diff --git a/Engine/Source/Editor/AudioEditor/AudioEditor.Build.cs b/Engine/Source/Editor/AudioEditor/AudioEditor.Build.cs
--- a/Engine/Source/Editor/AudioEditor/AudioEditor.Build.cs
+++ b/Engine/Source/Editor/AudioEditor/AudioEditor.Build.cs
@@ -11,15 +11,15 @@ public class AudioEditor : ModuleRules
 		(
 			new string[]
 			{
+				"AudioExtensions",
+				"AudioLinkEngine",
+				"DeveloperSettings",
 				"EditorSubsystem",
 				"GameProjectGeneration",
 				"ToolMenus",
+				"ToolWidgets",
 				"UMG",
-				"DeveloperSettings",
 				"UMGEditor",
-				"AudioExtensions",
-				"AudioLinkEngine",
-				"ToolWidgets"
 			}
 		);
 
@@ -27,26 +27,27 @@ public class AudioEditor : ModuleRules
 		(
 			new string[]
 			{
-				"Core",
-				"CoreUObject",
 				"ApplicationCore",
 				"AssetDefinition",
 				"AudioMixer",
-				"SignalProcessing",
-				"InputCore",
-				"Engine",
+				"ClassViewer",
+				"ContentBrowser",
+				"Core",
+				"CoreUObject",
+				"DetailCustomizations",
 				"EditorFramework",
-				"UnrealEd",
-				"Slate",
-				"SlateCore",
-				"RenderCore",
-				"LevelEditor",
+				"Engine",
+				"GraphEditor",
+				"InputCore",
+				"Kismet",
 				"Landscape",
+				"LevelEditor",
 				"PropertyEditor",
-				"DetailCustomizations",
-				"ClassViewer",
-				"GraphEditor",
-				"ContentBrowser",
+				"RenderCore",
+				"SignalProcessing",
+				"Slate",
+				"SlateCore",
+				"UnrealEd",
 			}
 		);
 
diff --git a/Engine/Source/Editor/AudioEditor/Private/AssetTypeActions/AssetDefinitionDefault_AudioDiffable.cpp b/Engine/Source/Editor/AudioEditor/Private/AssetTypeActions/AssetDefinitionDefault_AudioDiffable.cpp
new file mode 100644
--- /dev/null
+++ b/Engine/Source/Editor/AudioEditor/Private/AssetTypeActions/AssetDefinitionDefault_AudioDiffable.cpp
@@ -0,0 +1,23 @@
+// Copyright Epic Games, Inc. All Rights Reserved.
+
+#include "AssetTypeActions/AssetDefinitionDefault_AudioDiffable.h"
+
+#include "SDetailsDiff.h"
+
+#include UE_INLINE_GENERATED_CPP_BY_NAME(AssetDefinitionDefault_AudioDiffable)
+
+EAssetCommandResult UAssetDefinitionDefault_AudioDiffable::PerformAssetDiff(const FAssetDiffArgs& DiffArgs) const
+{
+	if (DiffArgs.OldAsset == nullptr && DiffArgs.NewAsset == nullptr)
+	{
+		return EAssetCommandResult::Unhandled;
+	}
+	
+	const TSharedRef<SDetailsDiff> DetailsDiff = SDetailsDiff::CreateDiffWindow(DiffArgs.OldAsset, DiffArgs.NewAsset, DiffArgs.OldRevision, DiffArgs.NewRevision, GetAssetClass().Get());
+	// allow users to edit NewAsset if it's a local asset
+	if (DiffArgs.NewAsset != nullptr && !FPackageName::IsTempPackage(DiffArgs.NewAsset->GetPackage()->GetName()))
+	{
+		DetailsDiff->SetOutputObject(DiffArgs.NewAsset);
+	}
+	return EAssetCommandResult::Handled;
+}
diff --git a/Engine/Source/Editor/AudioEditor/Private/AssetTypeActions/AssetDefinition_AudioBus.h b/Engine/Source/Editor/AudioEditor/Private/AssetTypeActions/AssetDefinition_AudioBus.h
--- a/Engine/Source/Editor/AudioEditor/Private/AssetTypeActions/AssetDefinition_AudioBus.h
+++ b/Engine/Source/Editor/AudioEditor/Private/AssetTypeActions/AssetDefinition_AudioBus.h
@@ -2,14 +2,14 @@
 
 #pragma once
 
+#include "AssetTypeActions/AssetDefinitionDefault_AudioDiffable.h"
 #include "CoreMinimal.h"
 #include "Sound/AudioBus.h"
-#include "AssetDefinitionDefault.h"
 
 #include "AssetDefinition_AudioBus.generated.h"
 
 UCLASS()
-class UAssetDefinition_AudioBus : public UAssetDefinitionDefault
+class UAssetDefinition_AudioBus : public UAssetDefinitionDefault_AudioDiffable
 {
 	GENERATED_BODY()
 
diff --git a/Engine/Source/Editor/AudioEditor/Private/AssetTypeActions/AssetDefinition_DialogueVoice.h b/Engine/Source/Editor/AudioEditor/Private/AssetTypeActions/AssetDefinition_DialogueVoice.h
--- a/Engine/Source/Editor/AudioEditor/Private/AssetTypeActions/AssetDefinition_DialogueVoice.h
+++ b/Engine/Source/Editor/AudioEditor/Private/AssetTypeActions/AssetDefinition_DialogueVoice.h
@@ -2,14 +2,14 @@
 
 #pragma once
 
+#include "AssetTypeActions/AssetDefinitionDefault_AudioDiffable.h"
 #include "CoreMinimal.h"
 #include "Sound/DialogueVoice.h"
-#include "AssetDefinitionDefault.h"
 
 #include "AssetDefinition_DialogueVoice.generated.h"
 
 UCLASS()
-class UAssetDefinition_DialogueVoice : public UAssetDefinitionDefault
+class UAssetDefinition_DialogueVoice : public UAssetDefinitionDefault_AudioDiffable
 {
 	GENERATED_BODY()
 
diff --git a/Engine/Source/Editor/AudioEditor/Private/AssetTypeActions/AssetDefinition_ReverbEffect.h b/Engine/Source/Editor/AudioEditor/Private/AssetTypeActions/AssetDefinition_ReverbEffect.h
--- a/Engine/Source/Editor/AudioEditor/Private/AssetTypeActions/AssetDefinition_ReverbEffect.h
+++ b/Engine/Source/Editor/AudioEditor/Private/AssetTypeActions/AssetDefinition_ReverbEffect.h
@@ -2,14 +2,14 @@
 
 #pragma once
 
+#include "AssetTypeActions/AssetDefinitionDefault_AudioDiffable.h"
 #include "CoreMinimal.h"
 #include "Sound/ReverbEffect.h"
-#include "AssetDefinitionDefault.h"
 
 #include "AssetDefinition_ReverbEffect.generated.h"
 
 UCLASS()
-class UAssetDefinition_ReverbEffect : public UAssetDefinitionDefault
+class UAssetDefinition_ReverbEffect : public UAssetDefinitionDefault_AudioDiffable
 {
 	GENERATED_BODY()
 
diff --git a/Engine/Source/Editor/AudioEditor/Private/AssetTypeActions/AssetDefinition_SoundAttenuation.h b/Engine/Source/Editor/AudioEditor/Private/AssetTypeActions/AssetDefinition_SoundAttenuation.h
--- a/Engine/Source/Editor/AudioEditor/Private/AssetTypeActions/AssetDefinition_SoundAttenuation.h
+++ b/Engine/Source/Editor/AudioEditor/Private/AssetTypeActions/AssetDefinition_SoundAttenuation.h
@@ -2,15 +2,15 @@
 
 #pragma once
 
+#include "AssetTypeActions/AssetDefinitionDefault_AudioDiffable.h"
 #include "CoreMinimal.h"
 #include "Sound/SoundAttenuation.h"
-#include "AssetDefinitionDefault.h"
 #include "AudioEditorSettings.h"
 
 #include "AssetDefinition_SoundAttenuation.generated.h"
 
 UCLASS()
-class UAssetDefinition_SoundAttenuation : public UAssetDefinitionDefault
+class UAssetDefinition_SoundAttenuation : public UAssetDefinitionDefault_AudioDiffable
 {
 	GENERATED_BODY()
 
diff --git a/Engine/Source/Editor/AudioEditor/Private/AssetTypeActions/AssetDefinition_SoundConcurrency.h b/Engine/Source/Editor/AudioEditor/Private/AssetTypeActions/AssetDefinition_SoundConcurrency.h
--- a/Engine/Source/Editor/AudioEditor/Private/AssetTypeActions/AssetDefinition_SoundConcurrency.h
+++ b/Engine/Source/Editor/AudioEditor/Private/AssetTypeActions/AssetDefinition_SoundConcurrency.h
@@ -2,15 +2,15 @@
 
 #pragma once
 
+#include "AssetTypeActions/AssetDefinitionDefault_AudioDiffable.h"
 #include "CoreMinimal.h"
 #include "Sound/SoundConcurrency.h"
-#include "AssetDefinitionDefault.h"
 #include "AudioEditorSettings.h"
 
 #include "AssetDefinition_SoundConcurrency.generated.h"
 
 UCLASS()
-class UAssetDefinition_SoundConcurrency : public UAssetDefinitionDefault
+class UAssetDefinition_SoundConcurrency : public UAssetDefinitionDefault_AudioDiffable
 {
 	GENERATED_BODY()
 
diff --git a/Engine/Source/Editor/AudioEditor/Private/AssetTypeActions/AssetDefinition_SoundMix.h b/Engine/Source/Editor/AudioEditor/Private/AssetTypeActions/AssetDefinition_SoundMix.h
--- a/Engine/Source/Editor/AudioEditor/Private/AssetTypeActions/AssetDefinition_SoundMix.h
+++ b/Engine/Source/Editor/AudioEditor/Private/AssetTypeActions/AssetDefinition_SoundMix.h
@@ -2,14 +2,14 @@
 
 #pragma once
 
+#include "AssetTypeActions/AssetDefinitionDefault_AudioDiffable.h"
 #include "CoreMinimal.h"
 #include "Sound/SoundMix.h"
-#include "AssetDefinitionDefault.h"
 
 #include "AssetDefinition_SoundMix.generated.h"
 
 UCLASS()
-class UAssetDefinition_SoundMix : public UAssetDefinitionDefault
+class UAssetDefinition_SoundMix : public UAssetDefinitionDefault_AudioDiffable
 {
 	GENERATED_BODY()
 
diff --git a/Engine/Source/Editor/AudioEditor/Private/AssetTypeActions/AssetTypeActions_AudioDiffable.cpp b/Engine/Source/Editor/AudioEditor/Private/AssetTypeActions/AssetTypeActions_AudioDiffable.cpp
new file mode 100644
--- /dev/null
+++ b/Engine/Source/Editor/AudioEditor/Private/AssetTypeActions/AssetTypeActions_AudioDiffable.cpp
@@ -0,0 +1,21 @@
+// Copyright Epic Games, Inc. All Rights Reserved.
+
+#include "AssetTypeActions/AssetTypeActions_AudioDiffable.h"
+
+#include "SDetailsDiff.h"
+
+void FAssetTypeActions_AudioDiffable::PerformAssetDiff(UObject* OldAsset, UObject* NewAsset, const struct FRevisionInfo& OldRevision,
+                                                   const struct FRevisionInfo& NewRevision) const
+{
+	if (OldAsset == nullptr && NewAsset == nullptr)
+	{
+		return;
+	}
+	
+	const TSharedRef<SDetailsDiff> DetailsDiff = SDetailsDiff::CreateDiffWindow(OldAsset, NewAsset, OldRevision, NewRevision, GetSupportedClass());
+	// allow users to edit NewAsset if it's a local asset
+	if (NewAsset && !FPackageName::IsTempPackage(NewAsset->GetPackage()->GetName()))
+	{
+		DetailsDiff->SetOutputObject(NewAsset);
+	}
+}
diff --git a/Engine/Source/Editor/AudioEditor/Private/AssetTypeActions/AssetTypeActions_DialogueWave.h b/Engine/Source/Editor/AudioEditor/Private/AssetTypeActions/AssetTypeActions_DialogueWave.h
--- a/Engine/Source/Editor/AudioEditor/Private/AssetTypeActions/AssetTypeActions_DialogueWave.h
+++ b/Engine/Source/Editor/AudioEditor/Private/AssetTypeActions/AssetTypeActions_DialogueWave.h
@@ -2,7 +2,7 @@
 
 #pragma once
 
-#include "AssetTypeActions_Base.h"
+#include "AssetTypeActions/AssetTypeActions_AudioDiffable.h"
 #include "AssetTypeCategories.h"
 #include "Containers/Array.h"
 #include "HAL/Platform.h"
@@ -18,7 +18,7 @@ class UClass;
 class UObject;
 
 
-class FAssetTypeActions_DialogueWave : public FAssetTypeActions_Base
+class FAssetTypeActions_DialogueWave : public FAssetTypeActions_AudioDiffable
 {
 public:
 	// IAssetTypeActions Implementation
diff --git a/Engine/Source/Editor/AudioEditor/Private/AssetTypeActions/AssetTypeActions_SoundClass.h b/Engine/Source/Editor/AudioEditor/Private/AssetTypeActions/AssetTypeActions_SoundClass.h
--- a/Engine/Source/Editor/AudioEditor/Private/AssetTypeActions/AssetTypeActions_SoundClass.h
+++ b/Engine/Source/Editor/AudioEditor/Private/AssetTypeActions/AssetTypeActions_SoundClass.h
@@ -2,7 +2,7 @@
 
 #pragma once
 
-#include "AssetTypeActions_Base.h"
+#include "AssetTypeActions/AssetTypeActions_AudioDiffable.h"
 #include "AssetTypeCategories.h"
 #include "Containers/Array.h"
 #include "HAL/Platform.h"
@@ -18,7 +18,7 @@ class UClass;
 class UObject;
 class USoundClass;
 
-class FAssetTypeActions_SoundClass : public FAssetTypeActions_Base
+class FAssetTypeActions_SoundClass : public FAssetTypeActions_AudioDiffable
 {
 public:
 	// IAssetTypeActions Implementation
diff --git a/Engine/Source/Editor/AudioEditor/Private/AssetTypeActions/AssetTypeActions_SoundEffectPreset.h b/Engine/Source/Editor/AudioEditor/Private/AssetTypeActions/AssetTypeActions_SoundEffectPreset.h
--- a/Engine/Source/Editor/AudioEditor/Private/AssetTypeActions/AssetTypeActions_SoundEffectPreset.h
+++ b/Engine/Source/Editor/AudioEditor/Private/AssetTypeActions/AssetTypeActions_SoundEffectPreset.h
@@ -2,7 +2,7 @@
 
 #pragma once
 
-#include "AssetTypeActions_Base.h"
+#include "AssetTypeActions/AssetTypeActions_AudioDiffable.h"
 #include "AssetTypeCategories.h"
 #include "Containers/Array.h"
 #include "HAL/Platform.h"
@@ -21,7 +21,7 @@ class UClass;
 class UObject;
 
 
-class FAssetTypeActions_SoundEffectSubmixPreset : public FAssetTypeActions_Base
+class FAssetTypeActions_SoundEffectSubmixPreset : public FAssetTypeActions_AudioDiffable
 {
 public:
 	// IAssetTypeActions Implementation
@@ -32,7 +32,7 @@ public:
 	virtual uint32 GetCategories() override { return EAssetTypeCategories::Sounds; }
 };
 
-class FAssetTypeActions_SoundEffectSourcePreset : public FAssetTypeActions_Base
+class FAssetTypeActions_SoundEffectSourcePreset : public FAssetTypeActions_AudioDiffable
 {
 public:
 	// IAssetTypeActions Implementation
@@ -43,7 +43,7 @@ public:
 	virtual uint32 GetCategories() override { return EAssetTypeCategories::Sounds; }
 };
 
-class FAssetTypeActions_SoundEffectSourcePresetChain : public FAssetTypeActions_Base
+class FAssetTypeActions_SoundEffectSourcePresetChain : public FAssetTypeActions_AudioDiffable
 {
 public:
 	// IAssetTypeActions Implementation
@@ -54,7 +54,7 @@ public:
 	virtual uint32 GetCategories() override { return EAssetTypeCategories::Sounds; }
 };
 
-class FAssetTypeActions_SoundEffectPreset : public FAssetTypeActions_Base
+class FAssetTypeActions_SoundEffectPreset : public FAssetTypeActions_AudioDiffable
 {
 public:
 	FAssetTypeActions_SoundEffectPreset(USoundEffectPreset* InEffectPreset);
diff --git a/Engine/Source/Editor/AudioEditor/Private/AssetTypeActions/AssetTypeActions_SoundSubmix.h b/Engine/Source/Editor/AudioEditor/Private/AssetTypeActions/AssetTypeActions_SoundSubmix.h
--- a/Engine/Source/Editor/AudioEditor/Private/AssetTypeActions/AssetTypeActions_SoundSubmix.h
+++ b/Engine/Source/Editor/AudioEditor/Private/AssetTypeActions/AssetTypeActions_SoundSubmix.h
@@ -2,7 +2,7 @@
 
 #pragma once
 
-#include "AssetTypeActions_Base.h"
+#include "AssetTypeActions/AssetTypeActions_AudioDiffable.h"
 #include "AssetTypeCategories.h"
 #include "Containers/Array.h"
 #include "HAL/Platform.h"
@@ -19,7 +19,7 @@ class UObject;
 
 /** Submix Types: */
 
-class FAssetTypeActions_SoundSubmix : public FAssetTypeActions_Base
+class FAssetTypeActions_SoundSubmix : public FAssetTypeActions_AudioDiffable
 {
 public:
 	// IAssetTypeActions Implementation
@@ -65,7 +65,7 @@ public:
 
 /** Base Settings factories */
 
-class FAssetTypeActions_SoundfieldEncodingSettings : public FAssetTypeActions_Base
+class FAssetTypeActions_SoundfieldEncodingSettings : public FAssetTypeActions_AudioDiffable
 {
 public:
 	// IAssetTypeActions Implementation
@@ -76,7 +76,7 @@ public:
 	virtual uint32 GetCategories() override { return EAssetTypeCategories::Sounds; }
 };
 
-class FAssetTypeActions_SoundfieldEffectSettings : public FAssetTypeActions_Base
+class FAssetTypeActions_SoundfieldEffectSettings : public FAssetTypeActions_AudioDiffable
 {
 public:
 	// IAssetTypeActions Implementation
@@ -87,7 +87,7 @@ public:
 	virtual uint32 GetCategories() override { return EAssetTypeCategories::Sounds; }
 };
 
-class FAssetTypeActions_SoundfieldEffect : public FAssetTypeActions_Base
+class FAssetTypeActions_SoundfieldEffect : public FAssetTypeActions_AudioDiffable
 {
 public:
 	// IAssetTypeActions Implementation
@@ -98,7 +98,7 @@ public:
 	virtual uint32 GetCategories() override { return EAssetTypeCategories::Sounds; }
 };
 
-class FAssetTypeActions_AudioEndpointSettings : public FAssetTypeActions_Base
+class FAssetTypeActions_AudioEndpointSettings : public FAssetTypeActions_AudioDiffable
 {
 public:
 	// IAssetTypeActions Implementation
@@ -109,7 +109,7 @@ public:
 	virtual uint32 GetCategories() override { return EAssetTypeCategories::Sounds; }
 };
 
-class FAssetTypeActions_SoundfieldEndpointSettings : public FAssetTypeActions_Base
+class FAssetTypeActions_SoundfieldEndpointSettings : public FAssetTypeActions_AudioDiffable
 {
 public:
 	// IAssetTypeActions Implementation
diff --git a/Engine/Source/Editor/AudioEditor/Public/AssetTypeActions/AssetDefinitionDefault_AudioDiffable.h b/Engine/Source/Editor/AudioEditor/Public/AssetTypeActions/AssetDefinitionDefault_AudioDiffable.h
new file mode 100644
--- /dev/null
+++ b/Engine/Source/Editor/AudioEditor/Public/AssetTypeActions/AssetDefinitionDefault_AudioDiffable.h
@@ -0,0 +1,18 @@
+// Copyright Epic Games, Inc. All Rights Reserved.
+#pragma once
+
+#include "AssetDefinitionDefault.h"
+
+#include "AssetDefinitionDefault_AudioDiffable.generated.h"
+
+UCLASS(Abstract, MinimalAPI)
+class UAssetDefinitionDefault_AudioDiffable : public UAssetDefinitionDefault
+{
+	GENERATED_BODY()
+
+public:
+	//~ Begin UAssetDefinition interface
+	AUDIOEDITOR_API virtual EAssetCommandResult PerformAssetDiff(const FAssetDiffArgs& DiffArgs) const override;
+	//~ End UAssetDefinition interface
+
+};
diff --git a/Engine/Source/Editor/AudioEditor/Public/AssetTypeActions/AssetDefinition_SoundBase.h b/Engine/Source/Editor/AudioEditor/Public/AssetTypeActions/AssetDefinition_SoundBase.h
--- a/Engine/Source/Editor/AudioEditor/Public/AssetTypeActions/AssetDefinition_SoundBase.h
+++ b/Engine/Source/Editor/AudioEditor/Public/AssetTypeActions/AssetDefinition_SoundBase.h
@@ -1,7 +1,7 @@
 // Copyright Epic Games, Inc. All Rights Reserved.
 #pragma once
 
-#include "AssetDefinitionDefault.h"
+#include "AssetDefinitionDefault_AudioDiffable.h"
 #include "ContentBrowserMenuContexts.h"
 
 #include "AssetDefinition_SoundBase.generated.h"
@@ -23,7 +23,7 @@ namespace UE::AudioEditor
 }
 
 UCLASS(Abstract, MinimalAPI)
-class UAssetDefinition_SoundAssetBase : public UAssetDefinitionDefault
+class UAssetDefinition_SoundAssetBase : public UAssetDefinitionDefault_AudioDiffable
 {
 	GENERATED_BODY()
 
diff --git a/Engine/Source/Editor/AudioEditor/Public/AssetTypeActions/AssetTypeActions_AudioDiffable.h b/Engine/Source/Editor/AudioEditor/Public/AssetTypeActions/AssetTypeActions_AudioDiffable.h
new file mode 100644
--- /dev/null
+++ b/Engine/Source/Editor/AudioEditor/Public/AssetTypeActions/AssetTypeActions_AudioDiffable.h
@@ -0,0 +1,12 @@
+// Copyright Epic Games, Inc. All Rights Reserved.
+#pragma once
+
+#include "AssetTypeActions_Base.h"
+
+class FAssetTypeActions_AudioDiffable : public FAssetTypeActions_Base
+{
+public:
+	//~ Begin FAssetTypeActions_Base interface
+	AUDIOEDITOR_API virtual void PerformAssetDiff(UObject* OldAsset, UObject* NewAsset, const struct FRevisionInfo& OldRevision, const struct FRevisionInfo& NewRevision) const override;;
+	//~ End FAssetTypeActions_Base interface
+};
diff --git a/Engine/Source/Editor/AudioEditor/Public/AssetTypeActions/AssetTypeActions_SoundBase.h b/Engine/Source/Editor/AudioEditor/Public/AssetTypeActions/AssetTypeActions_SoundBase.h
--- a/Engine/Source/Editor/AudioEditor/Public/AssetTypeActions/AssetTypeActions_SoundBase.h
+++ b/Engine/Source/Editor/AudioEditor/Public/AssetTypeActions/AssetTypeActions_SoundBase.h
@@ -2,7 +2,7 @@
 
 #pragma once
 
-#include "AssetTypeActions_Base.h"
+#include "AssetTypeActions_AudioDiffable.h"
 #include "AssetTypeCategories.h"
 #include "Containers/Array.h"
 #include "CoreMinimal.h"
@@ -25,7 +25,7 @@ struct FAssetData;
 class
 // UE_DEPRECATED(5.2, "The AssetDefinition system is replacing AssetTypeActions and UAssetDefinition_SoundBase replaced this.  Please see the Conversion Guide in AssetDefinition.h")
 
-FAssetTypeActions_SoundBase : public FAssetTypeActions_Base
+FAssetTypeActions_SoundBase : public FAssetTypeActions_AudioDiffable
 {
 public:
 	// IAssetTypeActions Implementation
