[SoundClassEditor] UX improvements
+ SoundClass nodes context menu now allow to find in content browser & view references (CTRL+B and ALT+SHIFT+R respectively)
+ Opened SoundClass is now focused in the graph

#Synced-CL 48352097

diff --git a/Engine/Source/Editor/AudioEditor/Private/SoundClassEditor.cpp b/Engine/Source/Editor/AudioEditor/Private/SoundClassEditor.cpp
--- a/Engine/Source/Editor/AudioEditor/Private/SoundClassEditor.cpp
+++ b/Engine/Source/Editor/AudioEditor/Private/SoundClassEditor.cpp
@@ -1,39 +1,39 @@
 // Copyright Epic Games, Inc. All Rights Reserved.
 
 #include "SoundClassEditor.h"
-#include "Modules/ModuleManager.h"
-#include "Styling/AppStyle.h"
+
+#include "AssetToolsModule.h"
+#include "Audio/AudioDebug.h"
+#include "AudioDeviceManager.h"
+#include "AudioEditorModule.h"
+#include "ContentBrowserModule.h"
 #include "EdGraph/EdGraph.h"
-#include "SoundClassGraph/SoundClassGraph.h"
-#include "SoundClassGraph/SoundClassGraphNode.h"
-#include "SoundClassGraph/SoundClassGraphSchema.h"
-#include "Sound/SoundClass.h"
 #include "Editor.h"
-#include "EngineAnalytics.h"
-
-#include "PropertyEditorModule.h"
-#include "IDetailsView.h"
-#include "AssetRegistry/AssetRegistryModule.h"
-
-#include "WorkspaceMenuStructureModule.h"
-
-#include "Kismet2/BlueprintEditorUtils.h"
-#include "ScopedTransaction.h"
-#include "SSoundClassActionMenu.h"
-#include "Widgets/Docking/SDockTab.h"
+#include "Factories/SoundClassFactory.h"
 #include "Framework/Application/SlateApplication.h"
 #include "Framework/Commands/GenericCommands.h"
 #include "Framework/MultiBox/MultiBoxBuilder.h"
-#include "AudioEditorModule.h"
-#include "AssetToolsModule.h"
+#include "Framework/MultiBox/MultiBoxBuilder.h"
+#include "Framework/Notifications/NotificationManager.h"
 #include "IAssetTools.h"
-#include "Factories/SoundClassFactory.h"
+#include "IContentBrowserSingleton.h"
+#include "IDetailsView.h"
+#include "Kismet2/BlueprintEditorUtils.h"
+#include "Modules/ModuleManager.h"
+#include "PropertyEditorModule.h"
+#include "SSoundClassActionMenu.h"
+#include "ScopedTransaction.h"
+#include "Sound/SoundClass.h"
+#include "SoundClassEditorCommands.h"
+#include "SoundClassGraph/SoundClassGraph.h"
+#include "SoundClassGraph/SoundClassGraphNode.h"
+#include "SoundClassGraph/SoundClassGraphSchema.h"
+#include "Styling/AppStyle.h"
+#include "Toolkits/GlobalEditorCommonCommands.h"
 #include "UObject/Package.h"
 #include "UObject/UObjectIterator.h"
-#include "SoundClassEditorCommands.h"
-#include "Audio/AudioDebug.h"
-#include "AudioDeviceManager.h"
-#include "Framework/MultiBox/MultiBoxBuilder.h"
+#include "Widgets/Docking/SDockTab.h"
+#include "Widgets/Notifications/SNotificationList.h"
 
 #define LOCTEXT_NAMESPACE "SoundClassEditor"
 DEFINE_LOG_CATEGORY_STATIC( LogSoundClassEditor, Log, All );
@@ -164,6 +164,7 @@ void FSoundClassEditor::InitSoundClassEditor( const EToolkitMode::Type Mode, con
 	}	
 #endif
 
+	bool bFoundClass = false;
 	GraphEditor->SelectAllNodes();
 	for (UObject* SelectedNode : GraphEditor->GetSelectedNodes())
 	{
@@ -172,10 +173,24 @@ void FSoundClassEditor::InitSoundClassEditor( const EToolkitMode::Type Mode, con
 		{
 			GraphEditor->ClearSelectionSet();
 			GraphEditor->SetNodeSelection(GraphNode, true);
+			GraphEditor->ZoomToFit(true);
 			DetailsView->SetObject(ObjectToEdit);
+			bFoundClass = true;
 			break;
 		}
 	}
+	
+	if (!bFoundClass)
+	{
+		GraphEditor->ClearSelectionSet();
+		
+		const FText NotificationErrorText = FText::Format(LOCTEXT("SoundClassMalformed", "SoundClass {0} has a parent, but this parent does not have it as child!"), FText::AsCultureInvariant(ObjectToEdit->GetName()));
+		FNotificationInfo Info(NotificationErrorText);
+		Info.ExpireDuration = 3.0f;
+		Info.bFireAndForget = true;
+		Info.Image = FAppStyle::Get().GetBrush("Icons.Error");
+		FSlateNotificationManager::Get().AddNotification(Info);
+	}
 }
 
 void FSoundClassEditor::BindCommands()
@@ -201,6 +216,13 @@ void FSoundClassEditor::BindCommands()
 		FExecuteAction::CreateSP(this, &FSoundClassEditor::ToggleMute),
 		FCanExecuteAction::CreateSP(this, &FSoundClassEditor::CanExcuteToggleMute),
 		FIsActionChecked::CreateSP(this, &FSoundClassEditor::IsMuteToggled));
+	
+	ToolkitCommands->MapAction(FGlobalEditorCommonCommands::Get().FindInContentBrowser, 
+		FExecuteAction::CreateSP(this, &FSoundClassEditor::FindInContentBrowser_Execute),
+		FCanExecuteAction::CreateSP(this, &FSoundClassEditor::CanFindInContentBrowser));
+	
+	ToolkitCommands->MapAction(FSoundClassEditorCommands::Get().ViewReferences,
+		FExecuteAction::CreateSP(this, &FSoundClassEditor::OnLaunchReferenceViewer));
 }
 
 FSoundClassEditor::FSoundClassEditor()
@@ -295,6 +317,14 @@ TSharedRef<SGraphEditor> FSoundClassEditor::CreateGraphEditorWidget()
 			FExecuteAction::CreateSP( this, &FSoundClassEditor::RemoveSelectedNodes ),
 			FCanExecuteAction::CreateSP( this, &FSoundClassEditor::CanRemoveNodes )
 			);
+		
+		// Node commands
+		GraphEditorCommands->MapAction(FGlobalEditorCommonCommands::Get().FindInContentBrowser, 
+			FExecuteAction::CreateSP(this, &FSoundClassEditor::FindInContentBrowser_Execute),
+			FCanExecuteAction::CreateSP(this, &FSoundClassEditor::CanFindInContentBrowser));
+	
+		GraphEditorCommands->MapAction(FSoundClassEditorCommands::Get().ViewReferences,
+			FExecuteAction::CreateSP(this, &FSoundClassEditor::OnLaunchReferenceViewer));
 	}
 
 	FGraphAppearanceInfo AppearanceInfo;
@@ -468,6 +498,55 @@ bool FSoundClassEditor::IsMuteToggled() const
 	return false;
 }
 
+void FSoundClassEditor::FindInContentBrowser_Execute()
+{
+	if (GraphEditor)
+	{
+		const FGraphPanelSelectionSet& CurrentSelection = GraphEditor->GetSelectedNodes();
+		if (CurrentSelection.Num() > 0)
+		{
+			TArray<FAssetData> AssetDatas;
+			AssetDatas.Reserve(CurrentSelection.Num());
+			for (UObject* Node : CurrentSelection)
+			{
+				USoundClassGraphNode* GraphNode = CastChecked<USoundClassGraphNode>(Node);
+				AssetDatas.Emplace(GraphNode->SoundClass);
+			}
+		
+			FContentBrowserModule& ContentBrowserModule = FModuleManager::Get().LoadModuleChecked<FContentBrowserModule>("ContentBrowser");
+			ContentBrowserModule.Get().SyncBrowserToAssets(AssetDatas);
+			return;
+		}
+	}
+	
+	// If no node selected or impossible to get to the graph, revert to basic toolkit behavior (find "opened" asset)
+	FAssetEditorToolkit::FindInContentBrowser_Execute();
+}
+
+void FSoundClassEditor::OnLaunchReferenceViewer() const
+{
+	TArray<FAssetIdentifier> Assets;
+	
+	const FGraphPanelSelectionSet& CurrentSelection = GraphEditor->GetSelectedNodes();
+	if (CurrentSelection.Num() > 0)
+	{
+		Assets.Reserve(CurrentSelection.Num());
+		for (UObject* Node : CurrentSelection)
+		{
+			USoundClassGraphNode* GraphNode = CastChecked<USoundClassGraphNode>(Node);
+			const FName AssetName(GraphNode->SoundClass->GetPackage()->GetFName());
+			Assets.Add(FAssetIdentifier(AssetName));
+		}
+		
+	}
+	else
+	{
+		Assets.Add(FAssetIdentifier(SoundClass->GetPackage()->GetFName()));
+	}
+	
+	FEditorDelegates::OnOpenReferenceViewer.Broadcast(Assets, FReferenceViewerParams());
+}
+
 void FSoundClassEditor::CreateSoundClass(UEdGraphPin* FromPin, const FVector2D& Location, const FString& Name)
 {
 	// If we have a valid name
diff --git a/Engine/Source/Editor/AudioEditor/Private/SoundClassEditor.h b/Engine/Source/Editor/AudioEditor/Private/SoundClassEditor.h
--- a/Engine/Source/Editor/AudioEditor/Private/SoundClassEditor.h
+++ b/Engine/Source/Editor/AudioEditor/Private/SoundClassEditor.h
@@ -4,12 +4,12 @@
 
 #include "AudioDefines.h"
 #include "CoreMinimal.h"
-#include "UObject/GCObject.h"
-#include "Widgets/DeclarativeSyntaxSupport.h"
-#include "Toolkits/IToolkitHost.h"
 #include "EditorUndoClient.h"
 #include "GraphEditor.h"
 #include "ISoundClassEditor.h"
+#include "Toolkits/IToolkitHost.h"
+#include "UObject/GCObject.h"
+#include "Widgets/DeclarativeSyntaxSupport.h"
 
 class IDetailsView;
 class UEdGraph;
@@ -63,7 +63,7 @@ public:
 	virtual FLinearColor GetWorldCentricTabColorScale() const override;
 
 	/** ISoundClassEditor interface */
-	void CreateSoundClass(class UEdGraphPin* FromPin, const FVector2D& Location, const FString& Name) override;
+	virtual void CreateSoundClass(class UEdGraphPin* FromPin, const FVector2D& Location, const FString& Name) override;
 
 	/** FEditorUndoClient Interface */
 	virtual void PostUndo(bool bSuccess) override;
@@ -110,6 +110,9 @@ private:
 	void ToggleMute();
 	bool CanExcuteToggleMute() const;
 	bool IsMuteToggled() const;
+	
+	virtual void FindInContentBrowser_Execute() override;
+	void OnLaunchReferenceViewer() const;
 
 private:
 	/** The SoundClass asset being inspected */
diff --git a/Engine/Source/Editor/AudioEditor/Private/SoundClassEditorCommands.cpp b/Engine/Source/Editor/AudioEditor/Private/SoundClassEditorCommands.cpp
--- a/Engine/Source/Editor/AudioEditor/Private/SoundClassEditorCommands.cpp
+++ b/Engine/Source/Editor/AudioEditor/Private/SoundClassEditorCommands.cpp
@@ -12,6 +12,7 @@ void FSoundClassEditorCommands::RegisterCommands()
 {
 	UI_COMMAND(ToggleSolo, "Solo", "Toggles Soloing this sound class", EUserInterfaceActionType::ToggleButton, FInputChord(EKeys::S));
 	UI_COMMAND(ToggleMute, "Mute", "Toggles Muting this sound class", EUserInterfaceActionType::ToggleButton, FInputChord(EKeys::M));
+	UI_COMMAND(ViewReferences, "Reference Viewer...", "Launches the reference viewer showing the selected asset references", EUserInterfaceActionType::Button, FInputChord(EModifierKey::Shift | EModifierKey::Alt, EKeys::R));
 }
 
 #undef LOCTEXT_NAMESPACE
diff --git a/Engine/Source/Editor/AudioEditor/Private/SoundClassGraphSchema.cpp b/Engine/Source/Editor/AudioEditor/Private/SoundClassGraphSchema.cpp
--- a/Engine/Source/Editor/AudioEditor/Private/SoundClassGraphSchema.cpp
+++ b/Engine/Source/Editor/AudioEditor/Private/SoundClassGraphSchema.cpp
@@ -11,6 +11,7 @@
 #include "HAL/PlatformCrt.h"
 #include "Internationalization/Internationalization.h"
 #include "ScopedTransaction.h"
+#include "SoundClassEditorCommands.h"
 #include "Sound/SoundClass.h"
 #include "SoundClassEditorUtilities.h"
 #include "SoundClassGraph/SoundClassGraph.h"
@@ -19,6 +20,7 @@
 #include "Templates/SharedPointer.h"
 #include "ToolMenu.h"
 #include "ToolMenuSection.h"
+#include "Toolkits/GlobalEditorCommonCommands.h"
 #include "UObject/ObjectPtr.h"
 
 #include UE_INLINE_GENERATED_CPP_BY_NAME(SoundClassGraphSchema)
@@ -56,16 +58,13 @@ void USoundClassGraphSchema::GetGraphContextActions(FGraphContextMenuBuilder& Co
 
 void USoundClassGraphSchema::GetContextMenuActions(UToolMenu* Menu, UGraphNodeContextMenuContext* Context) const
 {
-	if (Context->Node)
-	{
-		const USoundClassGraphNode* SoundGraphNode = Cast<const USoundClassGraphNode>(Context->Node);
-
-		{
-			FToolMenuSection& Section = Menu->AddSection("SoundClassGraphSchemaNodeActions", LOCTEXT("ClassActionsMenuHeader", "SoundClass Actions"));
-			Section.AddMenuEntry(FGraphEditorCommands::Get().BreakNodeLinks);
-			Section.AddMenuEntry(FGenericCommands::Get().Delete);
-		}
-	}
+	FToolMenuSection& Section = Menu->AddSection("SoundClassGraphSchemaNodeActions", LOCTEXT("ClassActionsMenuHeader", "SoundClass Actions"));
+	Section.AddMenuEntry(FGraphEditorCommands::Get().BreakNodeLinks);
+	Section.AddMenuEntry(FGenericCommands::Get().Delete);
+	
+	FToolMenuSection& AssetSection = Menu->AddSection("AssetGraphSchemaNodeActions", LOCTEXT("AssetActionsMenuHeader", "Asset Actions"));
+	AssetSection.AddMenuEntry(FGlobalEditorCommonCommands::Get().FindInContentBrowser);
+	AssetSection.AddMenuEntry(FSoundClassEditorCommands::Get().ViewReferences);
 
 	// No Super call so Node comments option is not shown
 }
diff --git a/Engine/Source/Editor/AudioEditor/Public/SoundClassEditorCommands.h b/Engine/Source/Editor/AudioEditor/Public/SoundClassEditorCommands.h
--- a/Engine/Source/Editor/AudioEditor/Public/SoundClassEditorCommands.h
+++ b/Engine/Source/Editor/AudioEditor/Public/SoundClassEditorCommands.h
@@ -31,6 +31,9 @@ public:
 	/** Plays the SoundCue or stops the currently playing cue/node */
 	TSharedPtr<FUICommandInfo> ToggleMute;
 
+	/** Shows the reference viewer for the selected assets */
+	TSharedPtr<FUICommandInfo> ViewReferences;
+
 	/** Initialize commands */
 	virtual void RegisterCommands() override;
 };
