[Audio] Sound Class graph search tab
+ Ctrl+F or Window menu to invoke
+ Keyboard navigation for results (search then can use arrow down/up directly)
+ Double-click or hit Space or Enter to zoom to the search result node

#tests Basic flow + removing nodes + undo/redo + closing & opening tab multiple times and retesting everyting
#Synced-CL 48352097

diff --git a/Engine/Source/Editor/AudioEditor/Private/SFindInSoundClassGraph.cpp b/Engine/Source/Editor/AudioEditor/Private/SFindInSoundClassGraph.cpp
new file mode 100644
--- /dev/null
+++ b/Engine/Source/Editor/AudioEditor/Private/SFindInSoundClassGraph.cpp
@@ -0,0 +1,267 @@
+// Copyright Epic Games, Inc. All Rights Reserved.
+
+#include "SFindInSoundClassGraph.h"
+
+#include "BlueprintUtilities.h"
+#include "Framework/Application/SlateApplication.h"
+#include "GraphEditAction.h"
+#include "GraphEditor.h"
+#include "Layout/WidgetPath.h"
+#include "Sound/SoundClass.h"
+#include "SoundClassGraph/SoundClassGraphNode.h"
+#include "Widgets/Input/SSearchBox.h"
+
+#define LOCTEXT_NAMESPACE "SoundClassEditor"
+
+void SFindInSoundClassGraph::Construct(const FArguments& InArgs, TSharedPtr<SGraphEditor> InGraphEditor)
+{
+	check(InGraphEditor.IsValid());
+	check(InGraphEditor->GetCurrentGraph() != nullptr);
+	
+	GraphEditorPtr = InGraphEditor;
+	OnResultActivated = InArgs._OnSearchResultActivated;
+	
+	// Update search results when the graph gets modified
+	InGraphEditor->GetCurrentGraph()->AddOnGraphChangedHandler(FOnGraphChanged::FDelegate::CreateSPLambda(this, [this](const FEdGraphEditAction& Action)
+	{
+		if (!SearchValue.IsEmpty())
+		{
+			// Last case is to detect "Undo/Redo" of node addition/creation, as these don't appear as AddNode or RemoveNode.
+			if (Action.Action == EEdGraphActionType::GRAPHACTION_AddNode || Action.Action == EEdGraphActionType::GRAPHACTION_RemoveNode 
+				|| (Action.Action == EEdGraphActionType::GRAPHACTION_Default && Action.Nodes.IsEmpty() && Action.Graph == nullptr))
+			{
+				InitiateSearch();
+			}
+		}
+	}));
+
+	this->ChildSlot
+	[
+		SNew(SVerticalBox)
+		+ SVerticalBox::Slot()
+		.AutoHeight()
+		.Padding(4.f)
+		[
+			SNew(SHorizontalBox)
+			+ SHorizontalBox::Slot()
+			.FillWidth(1)
+			[
+				SAssignNew(SearchTextField, SSearchBox)
+				.HintText(LOCTEXT("Find_GraphSearchHint", "Search"))
+				.OnTextChanged(this, &SFindInSoundClassGraph::OnSearchTextChanged)
+				.OnTextCommitted(this, &SFindInSoundClassGraph::OnSearchTextCommitted)
+				.DelayChangeNotificationsWhileTyping(true)
+				.DelayChangeNotificationsWhileTypingSeconds(0.25f)
+			]
+		]
+		+ SVerticalBox::Slot()
+		.AutoHeight()
+		.Padding(4.f)
+		[
+			SNew(STextBlock)
+			.Visibility_Lambda([this](){ return SearchTextField->GetText().IsEmptyOrWhitespace() ? EVisibility::Hidden : EVisibility::Visible; })
+			.Text_Lambda([this]() { return FText::Format(LOCTEXT("Find_NumResultsFmt", "{0} matching {0}|plural(one=node,other=nodes)"), FoundNodeCount); })
+		]
+		+ SVerticalBox::Slot()
+		.FillHeight(1.0f)
+		.Padding(4.f)
+		[
+			SNew(SBorder)
+			.BorderImage(FAppStyle::GetBrush("Menu.Background"))
+			[
+				SAssignNew(ListView, SListViewType)
+				.ListItemsSource(&ItemsFound)
+				.OnGenerateRow(this, &SFindInSoundClassGraph::OnGenerateRow)
+				.OnMouseButtonDoubleClick(this, &SFindInSoundClassGraph::OnSelectionDoubleClick)
+				.SelectionMode(ESelectionMode::Single)
+				.OnKeyDownHandler(FOnKeyDown::CreateSPLambda(this, [this](const FGeometry&, const FKeyEvent& Event)
+				{
+					if (Event.GetKey() == EKeys::Enter || Event.GetKey() == EKeys::SpaceBar)
+					{
+						TArray<FSearchResult> SelectedItems = ListView->GetSelectedItems();
+						if (!SelectedItems.IsEmpty())
+						{
+							OnSelectionDoubleClick(SelectedItems[0]);
+							return FReply::Handled();
+						}
+					}
+						
+					return FReply::Unhandled();
+				}))
+			]
+		]
+	];
+}
+
+void SFindInSoundClassGraph::FocusForUse()
+{
+	// NOTE: Careful, GeneratePathToWidget can be reentrant in that it can call visibility delegates and such
+	FWidgetPath FilterTextBoxWidgetPath;
+	FSlateApplication::Get().GeneratePathToWidgetUnchecked(SearchTextField.ToSharedRef(), FilterTextBoxWidgetPath);
+
+	// Set keyboard focus directly
+	FSlateApplication::Get().SetKeyboardFocus(FilterTextBoxWidgetPath, EFocusCause::SetDirectly);
+}
+
+void SFindInSoundClassGraph::FocusForUse(const FString& NewSearchTerms)
+{
+	FocusForUse();
+
+	if (!NewSearchTerms.IsEmpty())
+	{
+		SearchTextField->SetText(FText::FromString(NewSearchTerms));
+		InitiateSearch();
+	}
+}
+
+void SFindInSoundClassGraph::OnSearchTextChanged(const FText& Text)
+{
+	SearchValue = Text.ToString();
+	InitiateSearch();
+}
+
+void SFindInSoundClassGraph::OnSearchTextCommitted(const FText& Text, ETextCommit::Type CommitType)
+{
+	if (CommitType == ETextCommit::OnEnter)
+	{
+		SearchValue = Text.ToString();
+		InitiateSearch();
+	}
+}
+
+void SFindInSoundClassGraph::InitiateSearch()
+{
+	TArray<FString> Tokens;
+	if (SearchValue.Contains("\"") && SearchValue.ParseIntoArray(Tokens, TEXT("\""), true) > 0)
+	{
+		for (FString& TokenIt : Tokens)
+		{
+			// we have the token, we don't need the quotes anymore, they'll just confuse the comparison later on
+			TokenIt = TokenIt.TrimQuotes();
+			// We remove the spaces as all later comparison strings will also be de-spaced
+			TokenIt = TokenIt.Replace(TEXT(" "), TEXT(""));
+		}
+
+		// due to being able to handle multiple quoted blocks like ("Make Epic" "Game Now") we can end up with
+		// and empty string between (" ") blocks so this simply removes them
+		struct FRemoveMatchingStrings
+		{
+			bool operator()(const FString& RemovalCandidate) const
+			{
+				return RemovalCandidate.IsEmpty();
+			}
+		};
+		Tokens.RemoveAll(FRemoveMatchingStrings());
+	}
+	else
+	{
+		// unquoted search equivalent to a match-any-of search
+		SearchValue.ParseIntoArray(Tokens, TEXT(" "), true);
+	}
+
+	FoundNodeCount = 0;
+	ItemsFound.Empty();
+	if (Tokens.Num() > 0)
+	{
+		HighlightText = FText::FromString(SearchValue);
+		MatchTokens(Tokens);
+	}
+
+	struct FSearchResultLexicalLess
+	{
+		[[nodiscard]] bool operator()(const FSearchResult& Lhs, const FSearchResult& Rhs) const
+		{
+			// Push invalid results at the bottom, if any
+			if (!Lhs.IsValid())
+			{
+				return false;
+			}
+			if (!Rhs.IsValid())
+			{
+				return true;
+			}
+			
+			// Otherwise sort lexically ascending
+			return Lhs->SoundClass.GetName() < Rhs->SoundClass.GetName();
+		}
+	};
+	ItemsFound.StableSort(FSearchResultLexicalLess());
+	ListView->RequestListRefresh();
+}
+
+void SFindInSoundClassGraph::MatchTokens(const TArray<FString>& Tokens)
+{
+	if (TSharedPtr<SGraphEditor> GraphEditor = GraphEditorPtr.Pin())
+	{
+		MatchTokensInGraph(GraphEditor->GetCurrentGraph(), Tokens);
+	}
+}
+
+void SFindInSoundClassGraph::MatchTokensInGraph(const UEdGraph* Graph, const TArray<FString>& Tokens)
+{
+	if (!Graph)
+	{
+		return;
+	}
+	
+	for (auto It(Graph->Nodes.CreateConstIterator()); It; ++It)
+	{
+		UEdGraphNode* Node = *It;
+
+		FString NodeSearchString = Node->GetNodeTitle(ENodeTitleType::FullTitle).ToString().Replace(TEXT(" "), TEXT(""));
+
+		if (StringMatchesSearchTokens(Tokens, NodeSearchString))
+		{
+			ItemsFound.Emplace(CastChecked<USoundClassGraphNode>(Node));
+			FoundNodeCount++;
+		}
+	}
+
+	for (const UEdGraph* Subgraph : Graph->SubGraphs)
+	{
+		MatchTokensInGraph(Subgraph, Tokens);
+	}
+}
+
+TSharedRef<ITableRow> SFindInSoundClassGraph::OnGenerateRow(FSearchResult InItem, const TSharedRef<STableViewBase>& OwnerList)
+{
+	return SNew(STableRow<FSearchResult>, OwnerList)
+	[
+		SNew(SHorizontalBox)
+		+ SHorizontalBox::Slot()
+		.FillWidth(1.f)
+		.HAlign(HAlign_Left)
+		.VAlign(VAlign_Center)
+		.Padding(2.f)
+		[
+			SNew(STextBlock)
+				.Text(FText::FromString(GetNameSafe(InItem->SoundClass)))
+				.HighlightText_Lambda([this](){ return HighlightText; })
+		]
+	];
+}
+
+void SFindInSoundClassGraph::OnSelectionDoubleClick(FSearchResult Item)
+{
+	if (USoundClassGraphNode* Node = Item.Get())
+	{
+		OnResultActivated.ExecuteIfBound(Node);
+	}
+}
+
+bool SFindInSoundClassGraph::StringMatchesSearchTokens(const TArray<FString>& Tokens, const FString& ComparisonString)
+{
+	bool bFoundAllTokens = true;
+	//search the entry for each token, it must have all of them to pass
+	for (const FString& Token : Tokens)
+	{
+		if (!ComparisonString.Contains(Token))
+		{
+			bFoundAllTokens = false;
+			break;
+		}
+	}
+	return bFoundAllTokens;
+}
+
+#undef LOCTEXT_NAMESPACE
diff --git a/Engine/Source/Editor/AudioEditor/Private/SFindInSoundClassGraph.h b/Engine/Source/Editor/AudioEditor/Private/SFindInSoundClassGraph.h
new file mode 100644
--- /dev/null
+++ b/Engine/Source/Editor/AudioEditor/Private/SFindInSoundClassGraph.h
@@ -0,0 +1,80 @@
+// Copyright Epic Games, Inc. All Rights Reserved.
+
+#pragma once
+
+#include "Widgets/SCompoundWidget.h"
+#include "Widgets/Views/SListView.h"
+
+class SGraphEditor;
+class SSearchBox;
+class UEdGraph;
+class USoundClassGraphNode;
+
+class SFindInSoundClassGraph final : public SCompoundWidget
+{
+public:
+	DECLARE_DELEGATE_OneParam(FOnResultActivated, USoundClassGraphNode*)
+	
+	SLATE_BEGIN_ARGS(SFindInSoundClassGraph) {}
+		SLATE_EVENT(FOnResultActivated, OnSearchResultActivated)
+	SLATE_END_ARGS()
+
+	void Construct(const FArguments& InArgs, TSharedPtr<SGraphEditor> InGraphEditor);
+
+	/** Focuses this widget's search box */
+	void FocusForUse();
+	/** Focuses this widget's search box, and Initiates a search with given SearchTerms */
+	void FocusForUse(const FString& NewSearchTerms);
+
+private:
+	typedef TWeakObjectPtr<USoundClassGraphNode> FSearchResult;
+	typedef SListView<FSearchResult> SListViewType;
+
+	/** Called when user changes the text they are searching for */
+	void OnSearchTextChanged(const FText& Text);
+
+	/** Called when user changes commits text to the search box */
+	void OnSearchTextCommitted(const FText& Text, ETextCommit::Type CommitType);
+
+	/** Called when user double-clicks on a new result */
+	void OnSelectionDoubleClick(FSearchResult Item);
+
+	/** Called when a new row is being generated */
+	TSharedRef<ITableRow> OnGenerateRow(FSearchResult InItem, const TSharedRef<STableViewBase>& OwnerList);
+
+	/** Begins the search based on the SearchValue */
+	void InitiateSearch();
+
+	/** Find any results that contain all of the tokens */
+	void MatchTokens(const TArray<FString>& Tokens);
+
+	/** Find any results that contain all of the tokens in provided graph */
+	void MatchTokensInGraph(const UEdGraph* Graph, const TArray<FString>& Tokens);
+
+	/** Determines if a string matches the search tokens */
+	static bool StringMatchesSearchTokens(const TArray<FString>& Tokens, const FString& ComparisonString);
+
+	/** Pointer back to the Graph Editor of the Sound Class editor that owns us */
+	TWeakPtr<SGraphEditor> GraphEditorPtr;
+
+	/** The list view displays the results */
+	TSharedPtr<SListViewType> ListView;
+
+	/** The search text box */
+	TSharedPtr<SSearchBox> SearchTextField;
+	
+	/** The callback to execute when a search result is activated by the user (provides the node corresponding to the result) */
+	FOnResultActivated OnResultActivated;
+
+	/** This buffer stores the currently displayed results */
+	TArray<FSearchResult> ItemsFound;
+
+	/** The string to highlight in the results */
+	FText HighlightText;
+
+	/** The string to search for */
+	FString SearchValue;
+
+	/** Counters for number of results */
+	uint32 FoundNodeCount = 0;
+};
diff --git a/Engine/Source/Editor/AudioEditor/Private/SoundClassEditor.cpp b/Engine/Source/Editor/AudioEditor/Private/SoundClassEditor.cpp
--- a/Engine/Source/Editor/AudioEditor/Private/SoundClassEditor.cpp
+++ b/Engine/Source/Editor/AudioEditor/Private/SoundClassEditor.cpp
@@ -21,6 +21,7 @@
 #include "Kismet2/BlueprintEditorUtils.h"
 #include "Modules/ModuleManager.h"
 #include "PropertyEditorModule.h"
+#include "SFindInSoundClassGraph.h"
 #include "SSoundClassActionMenu.h"
 #include "ScopedTransaction.h"
 #include "Sound/SoundClass.h"
@@ -36,10 +37,11 @@
 #include "Widgets/Notifications/SNotificationList.h"
 
 #define LOCTEXT_NAMESPACE "SoundClassEditor"
-DEFINE_LOG_CATEGORY_STATIC( LogSoundClassEditor, Log, All );
+DEFINE_LOG_CATEGORY_STATIC(LogSoundClassEditor, Log, All);
 
-const FName FSoundClassEditor::GraphCanvasTabId( TEXT( "SoundClassEditor_GraphCanvas" ) );
-const FName FSoundClassEditor::PropertiesTabId( TEXT( "SoundClassEditor_Properties" ) );
+const FName FSoundClassEditor::GraphCanvasTabId(TEXT("SoundClassEditor_GraphCanvas"));
+const FName FSoundClassEditor::PropertiesTabId(TEXT("SoundClassEditor_Properties"));
+const FName FSoundClassEditor::FindInSoundClassGraphTabId(TEXT("SoundClassEditor_Find"));
 
 
 void FSoundClassEditor::ExtendToolbar()
@@ -77,23 +79,29 @@ void FSoundClassEditor::RegisterTabSpawners(const TSharedRef<class FTabManager>&
 
 	FAssetEditorToolkit::RegisterTabSpawners(InTabManager);
 
-	InTabManager->RegisterTabSpawner( GraphCanvasTabId, FOnSpawnTab::CreateSP(this, &FSoundClassEditor::SpawnTab_GraphCanvas) )
-		.SetDisplayName( LOCTEXT( "GraphCanvasTab", "Graph" ) )
-		.SetGroup( WorkspaceMenuCategoryRef )
+	InTabManager->RegisterTabSpawner(GraphCanvasTabId, FOnSpawnTab::CreateSP(this, &FSoundClassEditor::SpawnTab_GraphCanvas))
+		.SetDisplayName(LOCTEXT("GraphCanvasTab", "Graph"))
+		.SetGroup(WorkspaceMenuCategoryRef)
 		.SetIcon(FSlateIcon(FAppStyle::GetAppStyleSetName(), "GraphEditor.EventGraph_16x"));
 
-	InTabManager->RegisterTabSpawner( PropertiesTabId, FOnSpawnTab::CreateSP(this, &FSoundClassEditor::SpawnTab_Properties) )
-		.SetDisplayName( LOCTEXT( "PropertiesTab", "Details" ) )
-		.SetGroup( WorkspaceMenuCategoryRef )
+	InTabManager->RegisterTabSpawner(PropertiesTabId, FOnSpawnTab::CreateSP(this, &FSoundClassEditor::SpawnTab_Properties))
+		.SetDisplayName(LOCTEXT("PropertiesTab", "Details"))
+		.SetGroup(WorkspaceMenuCategoryRef)
 		.SetIcon(FSlateIcon(FAppStyle::GetAppStyleSetName(), "LevelEditor.Tabs.Details"));
+	
+	InTabManager->RegisterTabSpawner(FindInSoundClassGraphTabId, FOnSpawnTab::CreateSP(this, &FSoundClassEditor::SpawnTab_Find))
+		.SetDisplayName(LOCTEXT("FindTab", "Find in Sound Class Graph"))
+		.SetGroup(WorkspaceMenuCategoryRef)
+		.SetIcon(FSlateIcon(FAppStyle::GetAppStyleSetName(), "Kismet.Tabs.FindResults"));
 }
 
 void FSoundClassEditor::UnregisterTabSpawners(const TSharedRef<class FTabManager>& InTabManager)
 {
 	FAssetEditorToolkit::UnregisterTabSpawners(InTabManager);
 
-	InTabManager->UnregisterTabSpawner( GraphCanvasTabId );
-	InTabManager->UnregisterTabSpawner( PropertiesTabId );
+	InTabManager->UnregisterTabSpawner(GraphCanvasTabId);
+	InTabManager->UnregisterTabSpawner(PropertiesTabId);
+	InTabManager->UnregisterTabSpawner(FindInSoundClassGraphTabId);
 }
 
 void FSoundClassEditor::InitSoundClassEditor( const EToolkitMode::Type Mode, const TSharedPtr< class IToolkitHost >& InitToolkitHost, UObject* ObjectToEdit )
@@ -146,6 +154,9 @@ void FSoundClassEditor::InitSoundClassEditor( const EToolkitMode::Type Mode, con
 			)
 		)
 	);
+	
+	FindWidget = SNew(SFindInSoundClassGraph, GraphEditor)
+		.OnSearchResultActivated(this, &FSoundClassEditor::OnSearchResultActivated);
 
 	const bool bCreateDefaultStandaloneMenu = true;
 	const bool bCreateDefaultToolbar = true;
@@ -223,6 +234,9 @@ void FSoundClassEditor::BindCommands()
 	
 	ToolkitCommands->MapAction(FSoundClassEditorCommands::Get().ViewReferences,
 		FExecuteAction::CreateSP(this, &FSoundClassEditor::OnLaunchReferenceViewer));
+	
+	ToolkitCommands->MapAction(FSoundClassEditorCommands::Get().FindInSoundClassGraph,
+		FExecuteAction::CreateSP(this, &FSoundClassEditor::ShowFindInSoundClassGraph));
 }
 
 FSoundClassEditor::FSoundClassEditor()
@@ -246,7 +260,7 @@ void FSoundClassEditor::AddReferencedObjects( FReferenceCollector& Collector )
 
 TSharedRef<SDockTab> FSoundClassEditor::SpawnTab_GraphCanvas(const FSpawnTabArgs& Args)
 {
-	check( Args.GetTabId() == GraphCanvasTabId );
+	check(Args.GetTabId() == GraphCanvasTabId);
 
 	TSharedRef<SDockTab> SpawnedTab = SNew(SDockTab)
 		.Label(LOCTEXT("GraphCanvasTitle", "Graph"))
@@ -259,10 +273,10 @@ TSharedRef<SDockTab> FSoundClassEditor::SpawnTab_GraphCanvas(const FSpawnTabArgs
 
 TSharedRef<SDockTab> FSoundClassEditor::SpawnTab_Properties(const FSpawnTabArgs& Args)
 {
-	check( Args.GetTabId() == PropertiesTabId );
+	check(Args.GetTabId() == PropertiesTabId);
 
 	TSharedRef<SDockTab> SpawnedTab = SNew(SDockTab)
-		.Label( LOCTEXT( "SoundClassPropertiesTitle", "Details" ) )
+		.Label(LOCTEXT("SoundClassPropertiesTitle", "Details"))
 		[
 			DetailsView.ToSharedRef()
 		];
@@ -270,6 +284,19 @@ TSharedRef<SDockTab> FSoundClassEditor::SpawnTab_Properties(const FSpawnTabArgs&
 	return SpawnedTab;
 }
 
+TSharedRef<SDockTab> FSoundClassEditor::SpawnTab_Find(const FSpawnTabArgs& Args)
+{
+	check(Args.GetTabId() == FindInSoundClassGraphTabId);
+
+	TSharedRef<SDockTab> SpawnedTab = SNew(SDockTab)
+		.Label(LOCTEXT("SoundClassFindTitle", "Find Results")) 
+		[
+			FindWidget.ToSharedRef()
+		];
+	
+	return SpawnedTab;
+}
+
 FName FSoundClassEditor::GetToolkitFName() const
 {
 	return FName("SoundClassEditor");
@@ -547,6 +574,28 @@ void FSoundClassEditor::OnLaunchReferenceViewer() const
 	FEditorDelegates::OnOpenReferenceViewer.Broadcast(Assets, FReferenceViewerParams());
 }
 
+void FSoundClassEditor::ShowFindInSoundClassGraph()
+{
+	TabManager->TryInvokeTab(FindInSoundClassGraphTabId);
+	if (SFindInSoundClassGraph* Widget = FindWidget.Get())
+	{
+		Widget->FocusForUse();
+	}
+}
+
+void FSoundClassEditor::OnSearchResultActivated(USoundClassGraphNode* Node)
+{
+	GraphEditor->ClearSelectionSet();
+	
+	constexpr bool bSelect = true;
+	GraphEditor->SetNodeSelection(Node, bSelect);
+	
+	constexpr bool bOnlySelection = true;
+	GraphEditor->ZoomToFit(bOnlySelection);
+	
+	DetailsView->SetObject(Node->SoundClass);
+}
+
 void FSoundClassEditor::CreateSoundClass(UEdGraphPin* FromPin, const FVector2D& Location, const FString& Name)
 {
 	// If we have a valid name
diff --git a/Engine/Source/Editor/AudioEditor/Private/SoundClassEditor.h b/Engine/Source/Editor/AudioEditor/Private/SoundClassEditor.h
--- a/Engine/Source/Editor/AudioEditor/Private/SoundClassEditor.h
+++ b/Engine/Source/Editor/AudioEditor/Private/SoundClassEditor.h
@@ -14,6 +14,8 @@
 class IDetailsView;
 class UEdGraph;
 class USoundClass;
+class USoundClassGraphNode;
+class SFindInSoundClassGraph;
 
 namespace Audio
 {
@@ -72,6 +74,7 @@ public:
 private:
 	TSharedRef<SDockTab> SpawnTab_GraphCanvas(const FSpawnTabArgs& Args);
 	TSharedRef<SDockTab> SpawnTab_Properties(const FSpawnTabArgs& Args);
+	TSharedRef<SDockTab> SpawnTab_Find(const FSpawnTabArgs& Args);
 
 private:
 	/** Creates all internal widgets for the tabs to point at */
@@ -113,6 +116,9 @@ private:
 	
 	virtual void FindInContentBrowser_Execute() override;
 	void OnLaunchReferenceViewer() const;
+	
+	/** Show and focus the Find in Sound Class Graph tab. */
+	void ShowFindInSoundClassGraph();
 
 private:
 	/** The SoundClass asset being inspected */
@@ -123,6 +129,9 @@ private:
 
 	/** Property View */
 	TSharedPtr<class IDetailsView> DetailsView;
+	
+	/** Find in Sound Class Graph widget. */
+	TSharedPtr<SFindInSoundClassGraph> FindWidget;
 
 	/** Command list for this editor */
 	TSharedPtr<FUICommandList> GraphEditorCommands;
@@ -135,10 +144,14 @@ private:
 	/**	The tab ids for all the tabs used */
 	static const FName GraphCanvasTabId;
 	static const FName PropertiesTabId;
+	static const FName FindInSoundClassGraphTabId;
 
 	/** Bind to commands. */
 	void BindCommands();
 
 	/** Helper function to grow the toolbar and add custom buttons */
 	void ExtendToolbar();
+	
+	/** Select & zoom to result */
+	void OnSearchResultActivated(USoundClassGraphNode* Node);
 };
diff --git a/Engine/Source/Editor/AudioEditor/Private/SoundClassEditorCommands.cpp b/Engine/Source/Editor/AudioEditor/Private/SoundClassEditorCommands.cpp
--- a/Engine/Source/Editor/AudioEditor/Private/SoundClassEditorCommands.cpp
+++ b/Engine/Source/Editor/AudioEditor/Private/SoundClassEditorCommands.cpp
@@ -13,6 +13,7 @@ void FSoundClassEditorCommands::RegisterCommands()
 	UI_COMMAND(ToggleSolo, "Solo", "Toggles Soloing this sound class", EUserInterfaceActionType::ToggleButton, FInputChord(EKeys::S));
 	UI_COMMAND(ToggleMute, "Mute", "Toggles Muting this sound class", EUserInterfaceActionType::ToggleButton, FInputChord(EKeys::M));
 	UI_COMMAND(ViewReferences, "Reference Viewer...", "Launches the reference viewer showing the selected asset references", EUserInterfaceActionType::Button, FInputChord(EModifierKey::Shift | EModifierKey::Alt, EKeys::R));
+	UI_COMMAND(FindInSoundClassGraph, "Find in Sound Class Graph", "Find a Sound Class within this graph.", EUserInterfaceActionType::Button, FInputChord(EModifierKey::Control, EKeys::F));
 }
 
 #undef LOCTEXT_NAMESPACE
diff --git a/Engine/Source/Editor/AudioEditor/Public/SoundClassEditorCommands.h b/Engine/Source/Editor/AudioEditor/Public/SoundClassEditorCommands.h
--- a/Engine/Source/Editor/AudioEditor/Public/SoundClassEditorCommands.h
+++ b/Engine/Source/Editor/AudioEditor/Public/SoundClassEditorCommands.h
@@ -34,6 +34,9 @@ public:
 	/** Shows the reference viewer for the selected assets */
 	TSharedPtr<FUICommandInfo> ViewReferences;
 
+	/** Show the Find in Sound Class Graph window. */
+	TSharedPtr<FUICommandInfo> FindInSoundClassGraph;
+
 	/** Initialize commands */
 	virtual void RegisterCommands() override;
 };
