[BUGFIX] Sound Factory - Don't do all the processing if the sound wave duration is beyond the limit.

#rnx

#Synced-CL 49809867

diff --git a/Engine/Source/Editor/AudioEditor/Private/Factories/SoundFactory.cpp b/Engine/Source/Editor/AudioEditor/Private/Factories/SoundFactory.cpp
--- a/Engine/Source/Editor/AudioEditor/Private/Factories/SoundFactory.cpp
+++ b/Engine/Source/Editor/AudioEditor/Private/Factories/SoundFactory.cpp
@@ -424,14 +424,31 @@ UObject* USoundFactory::CreateObject
 			return nullptr;
 		}
 
+		// If we're a multi-channel file, we're going to spoof the behavior of the SoundSurroundFactory
+		int32 ChannelCount = (int32)*WaveInfo.pChannels;
+		check(ChannelCount > 0);
+
+		int32 SizeOfSample = (*WaveInfo.pBitsPerSample) / 8;
+
+		int32 NumSamples = WaveInfo.SampleDataSize / SizeOfSample;
+		int32 NumFrames = NumSamples / ChannelCount;
+
+		const bool bLimitingSoundWaveLength = SoundWaveImportLengthLimitInSecondsCVar > 0.0f;
+		const float SoundDuration = (float)NumFrames / *WaveInfo.pSamplesPerSec;
+		if (bLimitingSoundWaveLength && SoundDuration >= SoundWaveImportLengthLimitInSecondsCVar)
+		{
+			FMessageDialog::Open(EAppMsgType::Ok, FText::Format(NSLOCTEXT("SoundFactory", "Soundwave is too long to import"
+				, "{0} is {1} seconds in duration (this is over the limit of {2} seconds) {3}")
+				, FText::FromString(CuePackageName), SoundDuration, SoundWaveImportLengthLimitInSecondsCVar, Reason));
+
+			GEditor->GetEditorSubsystem<UImportSubsystem>()->BroadcastAssetPostImport(this, nullptr);
+			return nullptr;
+		}
+
 		// Use pre-existing sound if it exists and we want to keep settings,
 		// otherwise create new sound and import raw data.
 		USoundWave* Sound = (bUseExistingSettings && ExistingSound) ? ExistingSound : NewObject<USoundWave>(InParent, Name, Flags, TemplateSoundWave.Get());
 
-		// If we're a multi-channel file, we're going to spoof the behavior of the SoundSurroundFactory
-		int32 ChannelCount = (int32)*WaveInfo.pChannels;
-		check(ChannelCount >0);
-
 		// These get wiped in PostInitProperties by defaults set from Audio Settings,
 		// so set back to template in this specialized case
 		if (TemplateSoundWave.IsValid())
@@ -444,11 +461,6 @@ UObject* USoundFactory::CreateObject
 			Sound->ResetInvalidTemplateArguments();
 		}
 
-		int32 SizeOfSample = (*WaveInfo.pBitsPerSample) / 8;
-
-		int32 NumSamples = WaveInfo.SampleDataSize / SizeOfSample;
-		int32 NumFrames = NumSamples / ChannelCount;
-
 		// Starting or ending on a non-zero sample creates a pop unless you have very specific use cases
 		// (doesn't apply to ambi)
 		if (!bIsFuMa && !bIsAmbiX)
@@ -646,23 +658,12 @@ UObject* USoundFactory::CreateObject
 
 		}
 
-		Sound->Duration = (float)NumFrames / *WaveInfo.pSamplesPerSec;
+		Sound->Duration = SoundDuration;
 		Sound->SetImportedSampleRate(*WaveInfo.pSamplesPerSec);
 		Sound->SetSampleRate(*WaveInfo.pSamplesPerSec);
 		Sound->NumChannels = ChannelCount;
 		Sound->TotalSamples = *WaveInfo.pSamplesPerSec * Sound->Duration;
 
-		const bool bLimitingSoundWaveLength = SoundWaveImportLengthLimitInSecondsCVar > 0.0f; 
-		if (bLimitingSoundWaveLength && Sound->Duration >= SoundWaveImportLengthLimitInSecondsCVar)
-		{
-			FMessageDialog::Open(EAppMsgType::Ok, FText::Format(NSLOCTEXT("SoundFactory", "Soundwave is too long to import"
-				, "{0} is {1} seconds in duration (this is over the limit of {2} seconds) {3}")
-				, FText::FromString(CuePackageName), Sound->Duration, SoundWaveImportLengthLimitInSecondsCVar, Reason));
-				
-			GEditor->GetEditorSubsystem<UImportSubsystem>()->BroadcastAssetPostImport(this, nullptr);
-			return nullptr;			
-		}
-
 		// Store the current file path and timestamp for re-import purposes
 		Sound->AssetImportData->Update(CurrentFilename);
 
