Fixing up undo/redo in submix graph.
Fixed up detail customization in modulation destinations so they can be undone/redone

[RN] Fix for undo/redo transcations in the submix editor and in audio modulation destinations.

#Synced-CL 50410508

diff --git a/Engine/Source/Editor/AudioEditor/Private/SoundModulationDestinationLayout.cpp b/Engine/Source/Editor/AudioEditor/Private/SoundModulationDestinationLayout.cpp
--- a/Engine/Source/Editor/AudioEditor/Private/SoundModulationDestinationLayout.cpp
+++ b/Engine/Source/Editor/AudioEditor/Private/SoundModulationDestinationLayout.cpp
@@ -34,6 +34,7 @@
 #include "Misc/CString.h"
 #include "PropertyEditorModule.h"
 #include "PropertyHandle.h"
+#include "ScopedTransaction.h"
 #include "SlotBase.h"
 #include "Sound/SoundModulationDestination.h"
 #include "Sound/SoundSourceBus.h"
@@ -49,7 +50,7 @@
 #include "Widgets/Input/SButton.h"
 #include "Widgets/SBoxPanel.h"
 #include "Widgets/Text/STextBlock.h"
-
+#include "Widgets/Input/SNumericEntryBox.h"
 
 #define LOCTEXT_NAMESPACE "SoundModulationParameter"
 
@@ -154,7 +155,15 @@ namespace ModDestinationLayoutUtils
 		return false;
 	}
 
-	FText SetMetaData(TSharedRef<IPropertyHandle> StructPropertyHandle, TSharedRef<IPropertyHandle> ValueHandle, FText& OutUnitDisplayText, FName& OutParamName)
+	struct FNumericRanges
+	{
+		TOptional<float> ClampMinValue;
+		TOptional<float> ClampMaxValue;
+		TOptional<float> UIMinValue;
+		TOptional<float> UIMaxValue;
+	};
+
+	FText SetMetaData(TSharedRef<IPropertyHandle> StructPropertyHandle, TSharedRef<IPropertyHandle> ValueHandle, FText& OutUnitDisplayText, FName& OutParamName, FNumericRanges& OutRanges)
 	{
 		bool bClampValuesSet = false;
 		float ClampMinValue = 0.0f;
@@ -219,9 +228,35 @@ namespace ModDestinationLayoutUtils
 		ValueHandle->SetInstanceMetaData("UIMin", FString::Printf(TEXT("%f"), UIMinValue));
 		ValueHandle->SetInstanceMetaData("UIMax", FString::Printf(TEXT("%f"), UIMaxValue));
 		
+		OutRanges.ClampMinValue = ClampMinValue;
+		OutRanges.ClampMaxValue = ClampMaxValue;
+		OutRanges.UIMinValue = UIMinValue;
+		OutRanges.UIMaxValue = UIMaxValue;
+
 		return OutUnitDisplayText;
 	}
 
+	struct FValueEditState
+	{
+		TUniquePtr<FScopedTransaction> Transaction;
+		bool bIsDragging = false;
+		float LastCommittedValue = 0.f;
+	};
+
+	// Helper: modify all outer objects
+	static void ModifyOuterObjects(TSharedRef<IPropertyHandle> Handle)
+	{
+		TArray<UObject*> OuterObjects;
+		Handle->GetOuterObjects(OuterObjects);
+		for (UObject* Obj : OuterObjects)
+		{
+			if (Obj)
+			{
+				Obj->Modify();
+			}
+		}
+	}
+
 	void CustomizeChildren_AddValueRow(
 		IDetailChildrenBuilder& ChildBuilder,
 		TSharedRef<IPropertyHandle> StructPropertyHandle,
@@ -231,9 +266,52 @@ namespace ModDestinationLayoutUtils
 	{
 		FText UnitDisplayText = FText::GetEmpty();
 		FName ParamName;
-		SetMetaData(StructPropertyHandle, ValueHandle, UnitDisplayText, ParamName);
+		FNumericRanges Ranges;
+		SetMetaData(StructPropertyHandle, ValueHandle, UnitDisplayText, ParamName, Ranges);
 
 		const FText DisplayName = StructPropertyHandle->GetPropertyDisplayName();
+
+		// Keep state per-row (shared so lambdas stay valid)
+		TSharedRef<FValueEditState> EditState = MakeShared<FValueEditState>();
+
+		auto GetCurrentValue = [ValueHandle]() -> TOptional<float>
+		{
+			float V = 0.0f;
+			if (ValueHandle->GetValue(V) == FPropertyAccess::Success)
+			{
+				return V;
+			}
+			return TOptional<float>();
+		};
+
+		auto SetValueInteractive = [ValueHandle](float NewValue)
+		{
+			ValueHandle->SetValue(NewValue, EPropertyValueSetFlags::InteractiveChange);
+		};
+
+		auto SetValueCommitted = [ValueHandle](float NewValue)
+		{
+			ValueHandle->SetValue(NewValue, EPropertyValueSetFlags::DefaultFlags);
+		};
+
+		auto BeginTransaction = [EditState, ValueHandle]()
+		{
+			if (!EditState->Transaction)
+			{
+				EditState->Transaction = MakeUnique<FScopedTransaction>(
+					NSLOCTEXT("Audio Modulation", "ChangeModulationValue", "Change Modulation Value"));
+
+				// Record the outer objects into the transaction buffer so we can undo these changes
+				ModifyOuterObjects(ValueHandle);
+			}
+		};
+
+		auto EndTransaction = [EditState]()
+		{
+			EditState->Transaction.Reset();
+			EditState->bIsDragging = false;
+		};
+
 		ChildBuilder.AddCustomRow(DisplayName)
 		.NameContent()
 		[
@@ -251,7 +329,48 @@ namespace ModDestinationLayoutUtils
 				.Padding(4.0f, 0.0f, 0.0f, 0.0f)
 				.VAlign(VAlign_Center)
 				[
-					ValueHandle->CreatePropertyValueWidget()
+					SNew(SNumericEntryBox<float>)
+					.Value_Lambda(GetCurrentValue)
+					.MinValue(Ranges.ClampMinValue)
+					.MaxValue(Ranges.ClampMaxValue)
+					.MinSliderValue(Ranges.UIMinValue)
+					.MaxSliderValue(Ranges.UIMaxValue)
+					.AllowSpin(true)
+					.Font(IDetailLayoutBuilder::GetDetailFont())
+
+					// Drag start: open transaction & Modify()
+					.OnBeginSliderMovement_Lambda([EditState, BeginTransaction]()
+					{
+						EditState->bIsDragging = true;
+						BeginTransaction();
+					})
+
+					// Drag update: interactive set (keeps responsive updates)
+					.OnValueChanged_Lambda([EditState, BeginTransaction, SetValueInteractive](float NewValue)
+					{
+						// OnValueChanged can fire from typing too; only treat as interactive if dragging.
+						if (EditState->bIsDragging)
+						{
+							BeginTransaction();
+							SetValueInteractive(NewValue);
+						}
+					})
+
+					// Drag end: final committed write + close transaction
+					.OnEndSliderMovement_Lambda([EditState, BeginTransaction, SetValueCommitted, EndTransaction](float NewValue)
+					{
+						BeginTransaction();
+						SetValueCommitted(NewValue);
+						EndTransaction();
+					})
+
+					// Text commit (enter/tab/click away): commit write + close transaction
+					.OnValueCommitted_Lambda([BeginTransaction, SetValueCommitted, EndTransaction](float NewValue, ETextCommit::Type)
+					{
+						BeginTransaction();
+						SetValueCommitted(NewValue);
+						EndTransaction();
+					})
 				]
 			+ SHorizontalBox::Slot()
 				.AutoWidth()
@@ -314,8 +433,11 @@ namespace ModDestinationLayoutUtils
 				]
 		];
 
-		EnablementHandle->SetOnPropertyValueChanged(FSimpleDelegate::CreateLambda([EnablementHandle, ValueHandle, StructPropertyHandle, ModulatorsHandle]()
+		EnablementHandle->SetOnPropertyValueChanged(
+			FSimpleDelegate::CreateLambda([EnablementHandle, ValueHandle, StructPropertyHandle, ModulatorsHandle, BeginTransaction, EndTransaction]()
 		{
+			BeginTransaction();
+
 			bool bEnabled = false;
 			EnablementHandle->GetValue(bEnabled);
 			if (bEnabled)
@@ -324,10 +446,13 @@ namespace ModDestinationLayoutUtils
 				if (const Audio::FModulationParameter* Parameter = Audio::GetModulationParameterPtr(ParamName))
 				{
 					ValueHandle->SetValue(Parameter->DefaultValue);
+					EndTransaction();
 					return;
 				}
 			}
 			ModulatorsHandle->AsSet()->Empty();
+
+			EndTransaction();
 		}));
 	}
 
@@ -335,7 +460,8 @@ namespace ModDestinationLayoutUtils
 	{
 		FText UnitDisplayText = FText::GetEmpty();
 		FName ParamName;
-		SetMetaData(StructPropertyHandle, ValueHandle, UnitDisplayText, ParamName);
+		FNumericRanges Ranges;
+		SetMetaData(StructPropertyHandle, ValueHandle, UnitDisplayText, ParamName, Ranges);
 
 		const FText DisplayName = StructPropertyHandle->GetPropertyDisplayName();
 		FDetailWidgetRow& ValueNoModRow = ChildBuilder.AddCustomRow(DisplayName);
diff --git a/Engine/Source/Editor/AudioEditor/Private/SoundSubmixEditor.cpp b/Engine/Source/Editor/AudioEditor/Private/SoundSubmixEditor.cpp
--- a/Engine/Source/Editor/AudioEditor/Private/SoundSubmixEditor.cpp
+++ b/Engine/Source/Editor/AudioEditor/Private/SoundSubmixEditor.cpp
@@ -12,6 +12,7 @@
 #include "Editor/EditorEngine.h"
 #include "EngineAnalytics.h"
 #include "Factories/SoundSubmixFactory.h"
+#include "FileHelpers.h"
 #include "Framework/Commands/GenericCommands.h"
 #include "Framework/Commands/UIAction.h"
 #include "Framework/Commands/UICommandList.h"
@@ -115,19 +116,35 @@ void FSoundSubmixEditor::UnregisterTabSpawners(const TSharedRef<FTabManager>& In
 	InTabManager->UnregisterTabSpawner( PropertiesTabId );
 }
 
-void FSoundSubmixEditor::AddEditableSubmixChildren(USoundSubmixBase* RootSubmix)
+void FSoundSubmixEditor::AddEditableSubmixChildren(USoundSubmixBase* ParentSubmix)
 {
-	if (!RootSubmix)
+	if (!ParentSubmix)
 	{
 		return;
 	}
 
-	RootSubmix->SetFlags(RF_Transactional);
-	for (USoundSubmixBase* Child : RootSubmix->ChildSubmixes)
+	ParentSubmix->SetFlags(RF_Transactional);
+	for (USoundSubmixBase* Child : ParentSubmix->ChildSubmixes)
 	{
 		if (Child)
 		{
-			Child->SoundSubmixGraph = RootSubmix->SoundSubmixGraph;
+			// Only update the child submix if it's different
+			// this avoids unnecessarily dirtying submixes if they aren't actually changed
+			if (Child->SoundSubmixGraph != ParentSubmix->SoundSubmixGraph)
+			{
+				// Preserve child pacakge state for editor-only bookkeeping
+				UPackage* ChildPkg = Child->GetOutermost();
+				const bool bWasDirty = ChildPkg ? ChildPkg->IsDirty() : false;
+
+				Child->Modify(false);
+				Child->SoundSubmixGraph = ParentSubmix->SoundSubmixGraph;
+				Child->PostEditChange();
+
+				if (ChildPkg)
+				{
+					ChildPkg->SetDirtyFlag(bWasDirty);
+				}
+			}
 			AddEditingObject(Child);
 			AddEditableSubmixChildren(Child);
 		}
@@ -286,6 +303,7 @@ void FSoundSubmixEditor::CreateInternalWidgets(USoundSubmixBase* InSoundSubmix)
 	FPropertyEditorModule& PropertyEditorModule = FModuleManager::GetModuleChecked<FPropertyEditorModule>("PropertyEditor");
 	FDetailsViewArgs DetailsViewArgs;
 	DetailsViewArgs.NameAreaSettings = FDetailsViewArgs::HideNameArea;
+	DetailsViewArgs.NotifyHook = this;
 	DetailsView = PropertyEditorModule.CreateDetailView(DetailsViewArgs);
 	DetailsView->SetObject(InSoundSubmix);
 }
@@ -508,6 +526,48 @@ void FSoundSubmixEditor::AddMissingEditableSubmixes()
 	}
 }
 
+void FSoundSubmixEditor::NotifyPreChange(FProperty* PropertyAboutToChange)
+{
+	// When the property editor is in a transaction, make sure edited objects get recorded.
+	if (GUndo)
+	{
+		TArray<TWeakObjectPtr<UObject>> Edited;
+		Edited = DetailsView->GetSelectedObjects();
+
+		for (TWeakObjectPtr<UObject> Obj : Edited)
+		{
+			if (Obj.IsValid())
+			{
+				Obj->Modify();
+			}
+		}
+	}
+
+	FNotifyHook::NotifyPreChange(PropertyAboutToChange);
+}
+
+void FSoundSubmixEditor::SaveAsset_Execute()
+{
+	TArray<UPackage*> PackagesToSave;
+
+	for (UObject* Obj : GetEditingObjects())
+	{
+		if (!Obj) continue;
+		if (UPackage* Pkg = Obj->GetOutermost())
+		{
+			if (Pkg->IsDirty())
+			{
+				PackagesToSave.AddUnique(Pkg);
+			}
+		}
+	}
+
+	// Only attempt to save/checkout dirty ones.
+	const bool bCheckDirty = true;
+	const bool bPromptToSave = false;
+	FEditorFileUtils::PromptForCheckoutAndSave(PackagesToSave, bCheckDirty, bPromptToSave);
+}
+
 void FSoundSubmixEditor::PostUndo(bool bSuccess)
 {
 	GraphEditor->ClearSelectionSet();
diff --git a/Engine/Source/Editor/AudioEditor/Private/SoundSubmixEditor.h b/Engine/Source/Editor/AudioEditor/Private/SoundSubmixEditor.h
--- a/Engine/Source/Editor/AudioEditor/Private/SoundSubmixEditor.h
+++ b/Engine/Source/Editor/AudioEditor/Private/SoundSubmixEditor.h
@@ -11,6 +11,7 @@
 #include "Internationalization/Text.h"
 #include "Math/Color.h"
 #include "Math/Vector2D.h"
+#include "Misc/NotifyHook.h"
 #include "Templates/SharedPointer.h"
 #include "Toolkits/IToolkit.h"
 #include "UObject/GCObject.h"
@@ -29,7 +30,7 @@ class UEdGraphPin;
 class UObject;
 class USoundSubmixBase;
 
-class FSoundSubmixEditor : public ISoundSubmixEditor, public FGCObject, public FEditorUndoClient
+class FSoundSubmixEditor : public ISoundSubmixEditor, public FGCObject, public FEditorUndoClient, public FNotifyHook
 {
 public:
 	virtual void RegisterTabSpawners(const TSharedRef<FTabManager>& TabManager) override;
@@ -59,6 +60,7 @@ public:
 	virtual FText GetToolkitName() const override;
 	virtual FText GetToolkitToolTipText() const override;
 	virtual FString GetWorldCentricTabPrefix() const override;
+	virtual void SaveAsset_Execute() override;
 
 	/** @return Returns the color and opacity to use for the color that appears behind the tab text for this toolkit's tab in world-centric mode. */
 	virtual FLinearColor GetWorldCentricTabColorScale() const override;
@@ -67,6 +69,9 @@ public:
 	void CreateSoundSubmix(UEdGraphPin* FromPin, FVector2D Location, const FString& Name) override;
 	//~ End ISoundSubmixEditor
 
+	/** FNotifyHook Interface */
+	virtual void NotifyPreChange(FProperty* PropertyAboutToChange) override;
+
 	/** FEditorUndoClient Interface */
 	virtual void PostUndo(bool bSuccess) override;
 	virtual void PostRedo(bool bSuccess) override { PostUndo(bSuccess); }
@@ -95,8 +100,8 @@ private:
 	/** Called to create context menu when right-clicking on graph */
 	FActionMenuContent OnCreateGraphActionMenu(UEdGraph* InGraph, const FVector2f& InNodePosition, const TArray<UEdGraphPin*>& InDraggedPins, bool bAutoExpand, SGraphEditor::FActionMenuClosed InOnMenuClosed);
 
-	/** Adds all children of provided root submix as editable */
-	void AddEditableSubmixChildren(USoundSubmixBase* RootSubmix);
+	/** Adds all children of provided parent submix as editable */
+	void AddEditableSubmixChildren(USoundSubmixBase* ParentSubmix);
 
 	/** Select every node in the graph */
 	void SelectAllNodes();
