[SoundWave] Address found 'heap-buffer-overflow' in /mnt/horde/FNR+Main+Inc/Sync/Engine/Source/./Editor/AudioEditor/Private/Factories/SoundFactory.cpp
When NumFrames == 0, (NumFrames-1)*ChannelCount underflows, placing LastFrame before the start of the sample buffer
Added NumFrames > 0 guard so the DC-offset loop is skipped for empty audio data

#Synced-CL 52541261

diff --git a/Engine/Source/Editor/AudioEditor/Private/Factories/SoundFactory.cpp b/Engine/Source/Editor/AudioEditor/Private/Factories/SoundFactory.cpp
--- a/Engine/Source/Editor/AudioEditor/Private/Factories/SoundFactory.cpp
+++ b/Engine/Source/Editor/AudioEditor/Private/Factories/SoundFactory.cpp
@@ -463,7 +463,8 @@ UObject* USoundFactory::CreateObject
 
 		// Starting or ending on a non-zero sample creates a pop unless you have very specific use cases
 		// (doesn't apply to ambi)
-		if (!bIsFuMa && !bIsAmbiX)
+		// Skip the check for DC offsets when audio data is empty
+		if (NumFrames > 0 && !bIsFuMa && !bIsAmbiX)
 		{
 			int32 EndChannelNonZero = 0;
 			int32 StartChannelNonZero = 0;
