[ContentBrowser] New Add Menu Audio Menu

#Synced-CL 52619962

diff --git a/Engine/Plugins/Experimental/AudioProperties/Source/AudioPropertiesEditor/Private/AssetDefinition_AudioPropertiesBindings.cpp b/Engine/Plugins/Experimental/AudioProperties/Source/AudioPropertiesEditor/Private/AssetDefinition_AudioPropertiesBindings.cpp
--- a/Engine/Plugins/Experimental/AudioProperties/Source/AudioPropertiesEditor/Private/AssetDefinition_AudioPropertiesBindings.cpp
+++ b/Engine/Plugins/Experimental/AudioProperties/Source/AudioPropertiesEditor/Private/AssetDefinition_AudioPropertiesBindings.cpp
@@ -1,9 +1,18 @@
-﻿// Copyright Epic Games, Inc. All Rights Reserved.
+// Copyright Epic Games, Inc. All Rights Reserved.
 
 #include "AssetDefinition_AudioPropertiesBindings.h"
 
+#define LOCTEXT_NAMESPACE "AssetDefinition_AudioPropertiesBindings"
+
 TConstArrayView<FAssetCategoryPath> UAssetDefinition_AudioPropertiesBindings::GetAssetCategories() const
 {
-	static const auto Categories = { FAssetCategoryPath(EAssetCategoryPaths::Audio, NSLOCTEXT("AssetTypeActions", "AssetAudioPropertiesSubMenu", "Audio Properties")) };
+	static const auto Categories = 
+		{
+			FAssetCategoryPath(EAssetCategoryPaths::Audio, 
+			LOCTEXT("AssetAudioPropertiesBindingSubMenu", "Advanced"),
+				FCategoryPath(LOCTEXT("AssetAudioPropertiesBindingSubMenuSection", "Properties"), ECategoryMenuType::Section))
+		};
 	return Categories;
-}
\ No newline at end of file
+}
+
+#undef LOCTEXT_NAMESPACE
diff --git a/Engine/Plugins/Experimental/AudioProperties/Source/AudioPropertiesEditor/Private/AssetDefinition_AudioPropertiesSheetAsset.cpp b/Engine/Plugins/Experimental/AudioProperties/Source/AudioPropertiesEditor/Private/AssetDefinition_AudioPropertiesSheetAsset.cpp
--- a/Engine/Plugins/Experimental/AudioProperties/Source/AudioPropertiesEditor/Private/AssetDefinition_AudioPropertiesSheetAsset.cpp
+++ b/Engine/Plugins/Experimental/AudioProperties/Source/AudioPropertiesEditor/Private/AssetDefinition_AudioPropertiesSheetAsset.cpp
@@ -1,4 +1,4 @@
-﻿// Copyright Epic Games, Inc. All Rights Reserved.
+// Copyright Epic Games, Inc. All Rights Reserved.
 
 #include "AssetDefinition_AudioPropertiesSheetAsset.h"
 
@@ -11,7 +11,12 @@
 
 TConstArrayView<FAssetCategoryPath> UAssetDefinition_AudioPropertiesSheetAsset::GetAssetCategories() const
 {
-	static const auto Categories = { FAssetCategoryPath(EAssetCategoryPaths::Audio, NSLOCTEXT("AssetTypeActions", "AssetAudioPropertiesSubMenu", "Audio Properties")) };
+	static const auto Categories = 
+	{
+		FAssetCategoryPath(EAssetCategoryPaths::Audio, 
+			LOCTEXT("AssetAudioPropertiesSheetSubMenu", "Advanced"),
+			FCategoryPath(LOCTEXT("AssetAudioPropertiesSheetSubMenuSection", "Properties"), ECategoryMenuType::Section))
+	};
 	return Categories;
 }
 
@@ -55,4 +60,4 @@ TArray<FToolMenuSection*> UAssetDefinition_AudioPropertiesSheetAsset::RebuildCon
 	return Sections;
 }
 
-#undef LOCTEXT_NAMESPACE
\ No newline at end of file
+#undef LOCTEXT_NAMESPACE
diff --git a/Engine/Plugins/Experimental/MotoSynth/Source/MotoSynthEditor/Private/MotoSynthEditorModule.cpp b/Engine/Plugins/Experimental/MotoSynth/Source/MotoSynthEditor/Private/MotoSynthEditorModule.cpp
--- a/Engine/Plugins/Experimental/MotoSynth/Source/MotoSynthEditor/Private/MotoSynthEditorModule.cpp
+++ b/Engine/Plugins/Experimental/MotoSynth/Source/MotoSynthEditor/Private/MotoSynthEditorModule.cpp
@@ -15,7 +15,6 @@ void FMotoSynthEditorModule::StartupModule()
 	IAssetTools& AssetTools = FModuleManager::LoadModuleChecked<FAssetToolsModule>("AssetTools").Get();
 
 	AssetTools.RegisterAssetTypeActions(MakeShared<FAssetTypeActions_MotoSynthSource>());
-	AssetTools.RegisterAssetTypeActions(MakeShared<FAssetTypeActions_MotoSynthPreset>());
 
 	// Now that we've loaded this module, we need to register our effect preset actions
 	IAudioEditorModule* AudioEditorModule = &FModuleManager::LoadModuleChecked<IAudioEditorModule>("AudioEditor");
diff --git a/Engine/Plugins/Experimental/MotoSynth/Source/MotoSynthEditor/Private/MotoSynthSourceFactory.cpp b/Engine/Plugins/Experimental/MotoSynth/Source/MotoSynthEditor/Private/MotoSynthSourceFactory.cpp
--- a/Engine/Plugins/Experimental/MotoSynth/Source/MotoSynthEditor/Private/MotoSynthSourceFactory.cpp
+++ b/Engine/Plugins/Experimental/MotoSynth/Source/MotoSynthEditor/Private/MotoSynthSourceFactory.cpp
@@ -8,22 +8,7 @@
 
 #include UE_INLINE_GENERATED_CPP_BY_NAME(MotoSynthSourceFactory)
 
-#define LOCTEXT_NAMESPACE "AssetTypeActions"
-
-UClass* FAssetTypeActions_MotoSynthPreset::GetSupportedClass() const
-{
-	return UMotoSynthPreset::StaticClass();
-}
-
-const TArray<FText>& FAssetTypeActions_MotoSynthPreset::GetSubMenus() const
-{
-	static const TArray<FText> SubMenus
-	{
-		LOCTEXT("AssetSoundSynthesisSubMenu", "Synthesis")
-	};
-
-	return SubMenus;
-}
+#define LOCTEXT_NAMESPACE "MotoSynthSourceFactory"
 
 UMotoSynthPresetFactory::UMotoSynthPresetFactory(const FObjectInitializer& ObjectInitializer)
 	: Super(ObjectInitializer)
diff --git a/Engine/Plugins/Experimental/MotoSynth/Source/MotoSynthEditor/Private/MotoSynthSourceFactory.h b/Engine/Plugins/Experimental/MotoSynth/Source/MotoSynthEditor/Private/MotoSynthSourceFactory.h
--- a/Engine/Plugins/Experimental/MotoSynth/Source/MotoSynthEditor/Private/MotoSynthSourceFactory.h
+++ b/Engine/Plugins/Experimental/MotoSynth/Source/MotoSynthEditor/Private/MotoSynthSourceFactory.h
@@ -2,22 +2,35 @@
 
 #pragma once
 
+#include "AssetDefinitionDefault.h"
 #include "Factories/Factory.h"
 #include "AssetTypeActions_Base.h"
+#include "MotoSynthPreset.h"
 #include "MotoSynthSourceFactory.generated.h"
 
 class USoundWave;
 
-class FAssetTypeActions_MotoSynthPreset : public FAssetTypeActions_Base
+UCLASS()
+class UAssetDefinition_MotoSynthPreset : public UAssetDefinitionDefault
 {
+	GENERATED_BODY()
 public:
-	virtual FText GetName() const override { return NSLOCTEXT("AssetTypeActions", "AssetTypeActions_MotoSynthPreset", "Moto Synth Preset"); }
-	virtual FColor GetTypeColor() const override { return FColor(0, 150, 200); }
-	virtual UClass* GetSupportedClass() const override;
-	virtual uint32 GetCategories() override { return EAssetTypeCategories::Sounds; }
-	virtual const TArray<FText>& GetSubMenus() const override;
+	// UAssetDefinition Begin
+	virtual FText GetAssetDisplayName() const override { return NSLOCTEXT("AssetDefinition", "AssetDefinition_MotoSynthPreset", "Moto Synth Preset"); }
+	virtual TSoftClassPtr<UObject> GetAssetClass() const override { return UMotoSynthPreset::StaticClass(); }
+	virtual FLinearColor GetAssetColor() const override { return FColor(0, 150, 200); }
+	virtual TConstArrayView<FAssetCategoryPath> GetAssetCategories() const override
+	{
+		static const auto Categories = 
+		{
+			FAssetCategoryPath(EAssetCategoryPaths::Audio, NSLOCTEXT("AssetDefinition", "AssetDefinition_MotoSynthPresetSubMenu", "Legacy"))
+		};
+		return Categories;
+	}
+	// UAssetDefinition End
 };
 
+
 UCLASS(MinimalAPI, hidecategories = Object)
 class UMotoSynthPresetFactory : public UFactory
 {
diff --git a/Engine/Plugins/Experimental/Subsonic/Source/SubsonicEditor/Private/SubsonicAssetDefinitions.cpp b/Engine/Plugins/Experimental/Subsonic/Source/SubsonicEditor/Private/SubsonicAssetDefinitions.cpp
--- a/Engine/Plugins/Experimental/Subsonic/Source/SubsonicEditor/Private/SubsonicAssetDefinitions.cpp
+++ b/Engine/Plugins/Experimental/Subsonic/Source/SubsonicEditor/Private/SubsonicAssetDefinitions.cpp
@@ -61,7 +61,10 @@ namespace UE::Subsonic
 
 	TConstArrayView<FAssetCategoryPath> UAssetDefinition_SubsonicEventCollection::GetAssetCategories() const
 	{
-		static const TArray<FAssetCategoryPath> Categories = { EAssetCategoryPaths::Audio / LOCTEXT("AssetSoundSubsonicSubMenu", "Subsonic") };
+		static const auto Categories = 
+			{
+				FAssetCategoryPath(EAssetCategoryPaths::Audio, LOCTEXT("AssetSoundSubsonicSubMenu", "Experimental"))
+			};
 		return Categories;
 	}
 
diff --git a/Engine/Plugins/Runtime/AudioModulation/Source/AudioModulationEditor/AudioModulationEditor.build.cs b/Engine/Plugins/Runtime/AudioModulation/Source/AudioModulationEditor/AudioModulationEditor.build.cs
--- a/Engine/Plugins/Runtime/AudioModulation/Source/AudioModulationEditor/AudioModulationEditor.build.cs
+++ b/Engine/Plugins/Runtime/AudioModulation/Source/AudioModulationEditor/AudioModulationEditor.build.cs
@@ -15,6 +15,7 @@ public class AudioModulationEditor : ModuleRules
 		PrivateDependencyModuleNames.AddRange(
 			new string[]
 			{
+				"AssetDefinition",
 				"Core",
 				"CoreUObject",
 				"Engine",
diff --git a/Engine/Plugins/Runtime/AudioModulation/Source/AudioModulationEditor/Private/AssetDefinition/AssetDefinition_SoundControlBus.h b/Engine/Plugins/Runtime/AudioModulation/Source/AudioModulationEditor/Private/AssetDefinition/AssetDefinition_SoundControlBus.h
new file mode 100644
--- /dev/null
+++ b/Engine/Plugins/Runtime/AudioModulation/Source/AudioModulationEditor/Private/AssetDefinition/AssetDefinition_SoundControlBus.h
@@ -0,0 +1,40 @@
+// Copyright Epic Games, Inc. All Rights Reserved.
+
+#pragma once
+#include "AssetTypeActions/AssetDefinitionDefault_AudioDiffable.h"
+#include "AudioModulationStyle.h"
+#include "SoundControlBus.h"
+#include "AssetDefinition_SoundControlBus.generated.h"
+
+UCLASS()
+class UAssetDefinition_SoundControlBus : public UAssetDefinitionDefault_AudioDiffable
+{
+	GENERATED_BODY()
+
+public:
+	// UAssetDefinition Begin
+	virtual FText GetAssetDisplayName() const override { return NSLOCTEXT("AssetDefinition", "AssetDefinition_SoundControlBus", "Control Bus"); }
+	virtual FLinearColor GetAssetColor() const override { return UAudioModulationStyle::GetControlBusColor(); }
+	virtual TSoftClassPtr<UObject> GetAssetClass() const override { return USoundControlBus::StaticClass(); }
+	virtual FText GetAssetDescription(const FAssetData& AssetData) const override
+	{
+		if (const USoundControlBus* Bus = Cast<USoundControlBus>(AssetData.GetAsset()))
+		{
+			return FText::FromString(Bus->Description);
+		}
+		return FText::GetEmpty();
+	}
+
+	virtual TConstArrayView<FAssetCategoryPath> GetAssetCategories() const override
+	{
+		static const auto Categories = 
+		{
+			FAssetCategoryPath(EAssetCategoryPaths::Audio,
+				NSLOCTEXT("AssetDefinition", "AssetDefinition_SoundControlBusSubMenu", "Advanced"),
+				FCategoryPath(NSLOCTEXT("AssetDefinition", "AssetDefinition_SoundControlBusSubMenuSection", "Audio Modulation"), ECategoryMenuType::Section))
+		};
+
+		return Categories;
+	}
+	// UAssetDefinition End
+};
diff --git a/Engine/Plugins/Runtime/AudioModulation/Source/AudioModulationEditor/Private/AssetDefinition/AssetDefinition_SoundControlBusMix.h b/Engine/Plugins/Runtime/AudioModulation/Source/AudioModulationEditor/Private/AssetDefinition/AssetDefinition_SoundControlBusMix.h
new file mode 100644
--- /dev/null
+++ b/Engine/Plugins/Runtime/AudioModulation/Source/AudioModulationEditor/Private/AssetDefinition/AssetDefinition_SoundControlBusMix.h
@@ -0,0 +1,31 @@
+// Copyright Epic Games, Inc. All Rights Reserved.
+
+#pragma once
+#include "AssetTypeActions/AssetDefinitionDefault_AudioDiffable.h"
+#include "AudioModulationStyle.h"
+#include "SoundControlBusMix.h"
+#include "AssetDefinition_SoundControlBusMix.generated.h"
+
+UCLASS()
+class UAssetDefinition_SoundControlBusMix : public UAssetDefinitionDefault_AudioDiffable
+{
+	GENERATED_BODY()
+
+public:
+	// UAssetDefinition Begin
+	virtual FText GetAssetDisplayName() const override { return NSLOCTEXT("AssetDefinition", "AssetDefinition_SoundControlBusMix", "Control Bus Mix"); }
+	virtual FLinearColor GetAssetColor() const override { return UAudioModulationStyle::GetControlBusMixColor(); }
+	virtual TSoftClassPtr<UObject> GetAssetClass() const override { return USoundControlBusMix::StaticClass(); }
+	virtual TConstArrayView<FAssetCategoryPath> GetAssetCategories() const override
+	{
+		static const auto Categories = 
+		{
+			FAssetCategoryPath(EAssetCategoryPaths::Audio,
+				NSLOCTEXT("AssetDefinition", "AssetDefinition_SoundControlBusMixSubMenu", "Advanced"),
+				FCategoryPath(NSLOCTEXT("AssetDefinition", "AssetDefinition_SoundControlBusMixSubMenuSection", "Audio Modulation"), ECategoryMenuType::Section))
+		};
+
+		return Categories;
+	}
+	// UAssetDefinition End
+};
diff --git a/Engine/Plugins/Runtime/AudioModulation/Source/AudioModulationEditor/Private/AssetDefinition/AssetDefinition_SoundModulationGenerator.h b/Engine/Plugins/Runtime/AudioModulation/Source/AudioModulationEditor/Private/AssetDefinition/AssetDefinition_SoundModulationGenerator.h
new file mode 100644
--- /dev/null
+++ b/Engine/Plugins/Runtime/AudioModulation/Source/AudioModulationEditor/Private/AssetDefinition/AssetDefinition_SoundModulationGenerator.h
@@ -0,0 +1,31 @@
+// Copyright Epic Games, Inc. All Rights Reserved.
+
+#pragma once
+#include "AssetTypeActions/AssetDefinitionDefault_AudioDiffable.h"
+#include "AudioModulationStyle.h"
+#include "SoundModulationGenerator.h"
+#include "AssetDefinition_SoundModulationGenerator.generated.h"
+
+UCLASS()
+class UAssetDefinition_SoundModulationGenerator : public UAssetDefinitionDefault_AudioDiffable
+{
+	GENERATED_BODY()
+
+public:
+	// UAssetDefinition Begin
+	virtual FText GetAssetDisplayName() const override { return NSLOCTEXT("AssetDefinition", "AssetDefinition_SoundModulationGenerator", "Modulation Generator"); }
+	virtual FLinearColor GetAssetColor() const override { return UAudioModulationStyle::GetModulationGeneratorColor(); }
+	virtual TSoftClassPtr<UObject> GetAssetClass() const override { return USoundModulationGenerator::StaticClass(); }
+	virtual TConstArrayView<FAssetCategoryPath> GetAssetCategories() const override
+	{
+		static const auto Categories = 
+		{
+			FAssetCategoryPath(EAssetCategoryPaths::Audio,
+				NSLOCTEXT("AssetDefinition", "AssetDefinition_SoundModulationGeneratorSubMenu", "Advanced"),
+				FCategoryPath(NSLOCTEXT("AssetDefinition", "AssetDefinition_SoundModulationGeneratorSubMenuSection", "Audio Modulation"), ECategoryMenuType::Section))
+		};
+
+		return Categories;
+	}
+	// UAssetDefinition End
+};
diff --git a/Engine/Plugins/Runtime/AudioModulation/Source/AudioModulationEditor/Private/AssetDefinition/AssetDefinition_SoundModulationParameter.h b/Engine/Plugins/Runtime/AudioModulation/Source/AudioModulationEditor/Private/AssetDefinition/AssetDefinition_SoundModulationParameter.h
new file mode 100644
--- /dev/null
+++ b/Engine/Plugins/Runtime/AudioModulation/Source/AudioModulationEditor/Private/AssetDefinition/AssetDefinition_SoundModulationParameter.h
@@ -0,0 +1,31 @@
+// Copyright Epic Games, Inc. All Rights Reserved.
+
+#pragma once
+#include "AssetTypeActions/AssetDefinitionDefault_AudioDiffable.h"
+#include "AudioModulationStyle.h"
+#include "SoundModulationParameter.h"
+#include "AssetDefinition_SoundModulationParameter.generated.h"
+
+UCLASS()
+class UAssetDefinition_SoundModulationParameter : public UAssetDefinitionDefault_AudioDiffable
+{
+	GENERATED_BODY()
+
+public:
+	// UAssetDefinition Begin
+	virtual FText GetAssetDisplayName() const override { return NSLOCTEXT("AssetDefinition", "AssetDefinition_SoundModulationParameter", "Modulation Parameter"); }
+	virtual FLinearColor GetAssetColor() const override { return UAudioModulationStyle::GetParameterColor(); }
+	virtual TSoftClassPtr<UObject> GetAssetClass() const override { return USoundModulationParameter::StaticClass(); }
+	virtual TConstArrayView<FAssetCategoryPath> GetAssetCategories() const override
+	{
+		static const auto Categories = 
+		{
+			FAssetCategoryPath(EAssetCategoryPaths::Audio,
+				NSLOCTEXT("AssetDefinition", "AssetDefinition_SoundModulationParameterSubMenu", "Advanced"),
+				FCategoryPath(NSLOCTEXT("AssetDefinition", "AssetDefinition_SoundModulationParameterSubMenuSection", "Audio Modulation"), ECategoryMenuType::Section))
+		};
+
+		return Categories;
+	}
+	// UAssetDefinition End
+};
diff --git a/Engine/Plugins/Runtime/AudioModulation/Source/AudioModulationEditor/Private/AssetDefinition/AssetDefinition_SoundModulationPatch.cpp b/Engine/Plugins/Runtime/AudioModulation/Source/AudioModulationEditor/Private/AssetDefinition/AssetDefinition_SoundModulationPatch.cpp
new file mode 100644
--- /dev/null
+++ b/Engine/Plugins/Runtime/AudioModulation/Source/AudioModulationEditor/Private/AssetDefinition/AssetDefinition_SoundModulationPatch.cpp
@@ -0,0 +1,21 @@
+// Copyright Epic Games, Inc. All Rights Reserved.
+
+#include "AssetDefinition_SoundModulationPatch.h"
+
+#include "Editors/ModulationPatchEditor.h"
+
+EAssetCommandResult UAssetDefinition_SoundModulationPatch::OpenAssets(const FAssetOpenArgs& OpenArgs) const
+{
+	const EToolkitMode::Type Mode = OpenArgs.ToolkitHost.IsValid() ? EToolkitMode::WorldCentric : EToolkitMode::Standalone;
+
+	TArray<UObject*> Objects = OpenArgs.LoadObjects<UObject>();
+	for (UObject* Object : Objects)
+	{
+		if (USoundModulationPatch* Patch = Cast<USoundModulationPatch>(Object))
+		{
+			const TSharedRef<FModulationPatchEditor> PatchEditor = MakeShared<FModulationPatchEditor>();
+			PatchEditor->Init(Mode, OpenArgs.ToolkitHost, Patch);
+		}
+	}
+	return EAssetCommandResult::Handled;
+}
diff --git a/Engine/Plugins/Runtime/AudioModulation/Source/AudioModulationEditor/Private/AssetDefinition/AssetDefinition_SoundModulationPatch.h b/Engine/Plugins/Runtime/AudioModulation/Source/AudioModulationEditor/Private/AssetDefinition/AssetDefinition_SoundModulationPatch.h
new file mode 100644
--- /dev/null
+++ b/Engine/Plugins/Runtime/AudioModulation/Source/AudioModulationEditor/Private/AssetDefinition/AssetDefinition_SoundModulationPatch.h
@@ -0,0 +1,32 @@
+// Copyright Epic Games, Inc. All Rights Reserved.
+
+#pragma once
+#include "AssetTypeActions/AssetDefinitionDefault_AudioDiffable.h"
+#include "AudioModulationStyle.h"
+#include "SoundModulationPatch.h"
+#include "AssetDefinition_SoundModulationPatch.generated.h"
+
+UCLASS()
+class UAssetDefinition_SoundModulationPatch : public UAssetDefinitionDefault_AudioDiffable
+{
+	GENERATED_BODY()
+
+public:
+	// UAssetDefinition Begin
+	virtual FText GetAssetDisplayName() const override { return NSLOCTEXT("AssetDefinition", "AssetDefinition_SoundModulationPatch", "Parameter Patch"); }
+	virtual FLinearColor GetAssetColor() const override { return UAudioModulationStyle::GetPatchColor(); }
+	virtual TSoftClassPtr<UObject> GetAssetClass() const override { return USoundModulationPatch::StaticClass(); }
+	virtual TConstArrayView<FAssetCategoryPath> GetAssetCategories() const override
+	{
+		static const auto Categories = 
+		{
+			FAssetCategoryPath(EAssetCategoryPaths::Audio,
+				NSLOCTEXT("AssetDefinition", "AssetDefinition_SoundModulationPatchSubMenu", "Advanced"),
+				FCategoryPath(NSLOCTEXT("AssetDefinition", "AssetDefinition_SoundModulationPatchSubMenuSection", "Audio Modulation"), ECategoryMenuType::Section))
+		};
+
+		return Categories;
+	}
+	virtual EAssetCommandResult OpenAssets(const FAssetOpenArgs& OpenArgs) const override;
+	// UAssetDefinition End
+};
diff --git a/Engine/Plugins/Runtime/AudioModulation/Source/AudioModulationEditor/Private/AssetTypeActions/AssetTypeActions_SoundControlBus.cpp b/Engine/Plugins/Runtime/AudioModulation/Source/AudioModulationEditor/Private/AssetTypeActions/AssetTypeActions_SoundControlBus.cpp
deleted file mode 100644
diff --git a/Engine/Plugins/Runtime/AudioModulation/Source/AudioModulationEditor/Private/AssetTypeActions/AssetTypeActions_SoundControlBus.h b/Engine/Plugins/Runtime/AudioModulation/Source/AudioModulationEditor/Private/AssetTypeActions/AssetTypeActions_SoundControlBus.h
deleted file mode 100644
diff --git a/Engine/Plugins/Runtime/AudioModulation/Source/AudioModulationEditor/Private/AssetTypeActions/AssetTypeActions_SoundControlBusMix.cpp b/Engine/Plugins/Runtime/AudioModulation/Source/AudioModulationEditor/Private/AssetTypeActions/AssetTypeActions_SoundControlBusMix.cpp
deleted file mode 100644
diff --git a/Engine/Plugins/Runtime/AudioModulation/Source/AudioModulationEditor/Private/AssetTypeActions/AssetTypeActions_SoundControlBusMix.h b/Engine/Plugins/Runtime/AudioModulation/Source/AudioModulationEditor/Private/AssetTypeActions/AssetTypeActions_SoundControlBusMix.h
deleted file mode 100644
diff --git a/Engine/Plugins/Runtime/AudioModulation/Source/AudioModulationEditor/Private/AssetTypeActions/AssetTypeActions_SoundModulationGenerator.cpp b/Engine/Plugins/Runtime/AudioModulation/Source/AudioModulationEditor/Private/AssetTypeActions/AssetTypeActions_SoundModulationGenerator.cpp
deleted file mode 100644
diff --git a/Engine/Plugins/Runtime/AudioModulation/Source/AudioModulationEditor/Private/AssetTypeActions/AssetTypeActions_SoundModulationGenerator.h b/Engine/Plugins/Runtime/AudioModulation/Source/AudioModulationEditor/Private/AssetTypeActions/AssetTypeActions_SoundModulationGenerator.h
deleted file mode 100644
diff --git a/Engine/Plugins/Runtime/AudioModulation/Source/AudioModulationEditor/Private/AssetTypeActions/AssetTypeActions_SoundModulationParameter.cpp b/Engine/Plugins/Runtime/AudioModulation/Source/AudioModulationEditor/Private/AssetTypeActions/AssetTypeActions_SoundModulationParameter.cpp
deleted file mode 100644
diff --git a/Engine/Plugins/Runtime/AudioModulation/Source/AudioModulationEditor/Private/AssetTypeActions/AssetTypeActions_SoundModulationParameter.h b/Engine/Plugins/Runtime/AudioModulation/Source/AudioModulationEditor/Private/AssetTypeActions/AssetTypeActions_SoundModulationParameter.h
deleted file mode 100644
diff --git a/Engine/Plugins/Runtime/AudioModulation/Source/AudioModulationEditor/Private/AssetTypeActions/AssetTypeActions_SoundModulationPatch.cpp b/Engine/Plugins/Runtime/AudioModulation/Source/AudioModulationEditor/Private/AssetTypeActions/AssetTypeActions_SoundModulationPatch.cpp
deleted file mode 100644
diff --git a/Engine/Plugins/Runtime/AudioModulation/Source/AudioModulationEditor/Private/AssetTypeActions/AssetTypeActions_SoundModulationPatch.h b/Engine/Plugins/Runtime/AudioModulation/Source/AudioModulationEditor/Private/AssetTypeActions/AssetTypeActions_SoundModulationPatch.h
deleted file mode 100644
diff --git a/Engine/Plugins/Runtime/AudioModulation/Source/AudioModulationEditor/Private/AudioModulationEditor.cpp b/Engine/Plugins/Runtime/AudioModulation/Source/AudioModulationEditor/Private/AudioModulationEditor.cpp
--- a/Engine/Plugins/Runtime/AudioModulation/Source/AudioModulationEditor/Private/AudioModulationEditor.cpp
+++ b/Engine/Plugins/Runtime/AudioModulation/Source/AudioModulationEditor/Private/AudioModulationEditor.cpp
@@ -3,11 +3,6 @@
 
 #include "AssetRegistry/AssetRegistryModule.h"
 #include "AssetRegistry/IAssetRegistry.h"
-#include "AssetTypeActions/AssetTypeActions_SoundControlBus.h"
-#include "AssetTypeActions/AssetTypeActions_SoundControlBusMix.h"
-#include "AssetTypeActions/AssetTypeActions_SoundModulationGenerator.h"
-#include "AssetTypeActions/AssetTypeActions_SoundModulationParameter.h"
-#include "AssetTypeActions/AssetTypeActions_SoundModulationPatch.h"
 #include "AudioModulation.h"
 #include "Editors/ModulationPatchCurveEditorViewStacked.h"
 #include "Framework/Commands/UIAction.h"
@@ -93,15 +88,6 @@ void FAudioModulationEditorModule::StartupModule()
 	ModulationPatchToolBarExtensibilityManager = MakeShared<FExtensibilityManager>();
 	ModulationPatchMenuExtensibilityManager = MakeShared<FExtensibilityManager>();
 
-	// Register the audio editor asset type actions
-	IAssetTools& AssetTools = FModuleManager::LoadModuleChecked<FAssetToolsModule>(AudioModulationEditor::ToolName).Get();
-
-	AudioModulationEditor::AddAssetAction<FAssetTypeActions_SoundControlBus>(AssetTools, AssetActions);
-	AudioModulationEditor::AddAssetAction<FAssetTypeActions_SoundControlBusMix>(AssetTools, AssetActions);
-	AudioModulationEditor::AddAssetAction<FAssetTypeActions_SoundModulationGenerator>(AssetTools, AssetActions);
-	AudioModulationEditor::AddAssetAction<FAssetTypeActions_SoundModulationParameter>(AssetTools, AssetActions);
-	AudioModulationEditor::AddAssetAction<FAssetTypeActions_SoundModulationPatch>(AssetTools, AssetActions);
-
 	SetIcon(TEXT("SoundControlBus"));
 	SetIcon(TEXT("SoundControlBusMix"));
 	SetIcon(TEXT("SoundModulationGenerator"));
diff --git a/Engine/Plugins/Runtime/AudioSynesthesia/Source/AudioSynesthesia/Private/ConstantQ.cpp b/Engine/Plugins/Runtime/AudioSynesthesia/Source/AudioSynesthesia/Private/ConstantQ.cpp
--- a/Engine/Plugins/Runtime/AudioSynesthesia/Source/AudioSynesthesia/Private/ConstantQ.cpp
+++ b/Engine/Plugins/Runtime/AudioSynesthesia/Source/AudioSynesthesia/Private/ConstantQ.cpp
@@ -114,7 +114,7 @@ FText UConstantQSettings::GetAssetActionName() const
 
 UClass* UConstantQSettings::GetSupportedClass() const
 {
-	return UConstantQNRTSettings::StaticClass();
+	return UConstantQSettings::StaticClass();
 }
 #endif // WITH_EDITOR
 
diff --git a/Engine/Plugins/Runtime/AudioSynesthesia/Source/AudioSynesthesiaEditor/AudioSynesthesiaEditor.Build.cs b/Engine/Plugins/Runtime/AudioSynesthesia/Source/AudioSynesthesiaEditor/AudioSynesthesiaEditor.Build.cs
--- a/Engine/Plugins/Runtime/AudioSynesthesia/Source/AudioSynesthesiaEditor/AudioSynesthesiaEditor.Build.cs
+++ b/Engine/Plugins/Runtime/AudioSynesthesia/Source/AudioSynesthesiaEditor/AudioSynesthesiaEditor.Build.cs
@@ -24,9 +24,9 @@ public class AudioSynesthesiaEditor : ModuleRules
             }
         );
 
-		PrivateIncludePathModuleNames.AddRange(
+		PrivateDependencyModuleNames.AddRange(
 			new string[] {
-				"AssetTools",
+				"AssetDefinition"
 			});
 	}
 }
diff --git a/Engine/Plugins/Runtime/AudioSynesthesia/Source/AudioSynesthesiaEditor/Private/AssetDefinition_AudioSynesthesiaDynamic.cpp b/Engine/Plugins/Runtime/AudioSynesthesia/Source/AudioSynesthesiaEditor/Private/AssetDefinition_AudioSynesthesiaDynamic.cpp
new file mode 100644
--- /dev/null
+++ b/Engine/Plugins/Runtime/AudioSynesthesia/Source/AudioSynesthesiaEditor/Private/AssetDefinition_AudioSynesthesiaDynamic.cpp
@@ -0,0 +1,45 @@
+// Copyright Epic Games, Inc. All Rights Reserved.
+
+#include "AssetDefinition_AudioSynesthesiaDynamic.h"
+
+#define LOCTEXT_NAMESPACE "AssetDefinition_AudioSynesthesiaDynamic"
+
+FText UAssetDefinition_AudioSynesthesiaDynamic::GetAssetDisplayName() const
+{
+	if (SynesthesiaCDO)
+	{
+		return SynesthesiaCDO->GetAssetActionName();
+	}
+	return FText::FromString(TEXT("Audio Analyzer Asset Base"));
+}
+
+FLinearColor UAssetDefinition_AudioSynesthesiaDynamic::GetAssetColor() const
+{
+	if (SynesthesiaCDO)
+	{
+		return SynesthesiaCDO->GetTypeColor();
+	}
+	return FColor::White;
+}
+
+TSoftClassPtr<UObject> UAssetDefinition_AudioSynesthesiaDynamic::GetAssetClass() const
+{
+	if (SynesthesiaCDO)
+	{
+		if (UClass* SupportedClass = SynesthesiaCDO->GetSupportedClass())
+		{
+			return SupportedClass;
+		}
+
+		return SynesthesiaCDO->GetClass();
+	}
+
+	return UAudioAnalyzerAssetBase::StaticClass();
+}
+
+void UAssetDefinition_AudioSynesthesiaDynamic::Initialize(TSubclassOf<UAudioAnalyzerAssetBase> InClass)
+{
+	SynesthesiaCDO = TStrongObjectPtr(InClass->GetDefaultObject<UAudioAnalyzerAssetBase>());
+}
+
+#undef LOCTEXT_NAMESPACE
diff --git a/Engine/Plugins/Runtime/AudioSynesthesia/Source/AudioSynesthesiaEditor/Private/AssetDefinition_AudioSynesthesiaDynamic.h b/Engine/Plugins/Runtime/AudioSynesthesia/Source/AudioSynesthesiaEditor/Private/AssetDefinition_AudioSynesthesiaDynamic.h
new file mode 100644
--- /dev/null
+++ b/Engine/Plugins/Runtime/AudioSynesthesia/Source/AudioSynesthesiaEditor/Private/AssetDefinition_AudioSynesthesiaDynamic.h
@@ -0,0 +1,25 @@
+// Copyright Epic Games, Inc. All Rights Reserved.
+
+#pragma once
+#include "AudioAnalyzerAsset.h"
+#include "AssetTypeActions/AssetDefinitionDefault_AudioDiffable.h"
+#include "Templates/SubclassOf.h"
+#include "AssetDefinition_AudioSynesthesiaDynamic.generated.h"
+
+UCLASS()
+class UAssetDefinition_AudioSynesthesiaDynamic : public UAssetDefinitionDefault_AudioDiffable
+{
+	GENERATED_BODY()
+public:
+	// UAssetDefinition Begin
+	virtual FText GetAssetDisplayName() const override;
+	virtual FLinearColor GetAssetColor() const override;
+	virtual TSoftClassPtr<UObject> GetAssetClass() const override;
+	virtual bool CanRegisterStatically() const override { return false; }
+	// UAssetDefinition End
+
+	void Initialize(TSubclassOf<UAudioAnalyzerAssetBase> InClass);
+
+private:
+	TStrongObjectPtr<UAudioAnalyzerAssetBase> SynesthesiaCDO;
+};
diff --git a/Engine/Plugins/Runtime/AudioSynesthesia/Source/AudioSynesthesiaEditor/Private/AssetDefinition_AudioSynesthesiaNRT.h b/Engine/Plugins/Runtime/AudioSynesthesia/Source/AudioSynesthesiaEditor/Private/AssetDefinition_AudioSynesthesiaNRT.h
new file mode 100644
--- /dev/null
+++ b/Engine/Plugins/Runtime/AudioSynesthesia/Source/AudioSynesthesiaEditor/Private/AssetDefinition_AudioSynesthesiaNRT.h
@@ -0,0 +1,29 @@
+// Copyright Epic Games, Inc. All Rights Reserved.
+
+#pragma once
+#include "AssetDefinitionDefault.h"
+#include "AudioSynesthesiaNRT.h"
+#include "AssetTypeActions/AssetDefinitionDefault_AudioDiffable.h"
+#include "AssetDefinition_AudioSynesthesiaNRT.generated.h"
+
+UCLASS()
+class UAssetDefinition_AudioSynesthesiaNRT : public UAssetDefinitionDefault_AudioDiffable
+{
+	GENERATED_BODY()
+public:
+	// UAssetDefinition Begin
+	virtual FText GetAssetDisplayName() const override { return NSLOCTEXT("AssetDefinition", "AssetDefinition_AudioSynesthesiaNRT", "Synesthesia NRT"); }
+	virtual FLinearColor GetAssetColor() const override { return FColor(200.0f, 150.0f, 200.0f); }
+	virtual TSoftClassPtr<UObject> GetAssetClass() const override { return UAudioSynesthesiaNRT::StaticClass(); }
+	virtual TConstArrayView<FAssetCategoryPath> GetAssetCategories() const override
+	{
+		static const auto Categories = 
+			{
+				FAssetCategoryPath(EAssetCategoryPaths::Audio,
+					NSLOCTEXT("AssetDefinition", "AssetDefinition_AudioSynesthesiaNRTSubMenu", "Advanced"),
+					FCategoryPath(NSLOCTEXT("AssetDefinition", "AssetDefinition_AudioSynesthesiaNRTSubMenuSection", "Analysis"), ECategoryMenuType::Section))
+			};
+		return Categories;
+	}
+	// UAssetDefinition End
+};
diff --git a/Engine/Plugins/Runtime/AudioSynesthesia/Source/AudioSynesthesiaEditor/Private/AssetDefinition_AudioSynesthesiaNRTSettings.h b/Engine/Plugins/Runtime/AudioSynesthesia/Source/AudioSynesthesiaEditor/Private/AssetDefinition_AudioSynesthesiaNRTSettings.h
new file mode 100644
--- /dev/null
+++ b/Engine/Plugins/Runtime/AudioSynesthesia/Source/AudioSynesthesiaEditor/Private/AssetDefinition_AudioSynesthesiaNRTSettings.h
@@ -0,0 +1,28 @@
+// Copyright Epic Games, Inc. All Rights Reserved.
+
+#pragma once
+#include "AssetTypeActions/AssetDefinitionDefault_AudioDiffable.h"
+#include "AudioSynesthesiaNRT.h"
+#include "AssetDefinition_AudioSynesthesiaNRTSettings.generated.h"
+
+UCLASS()
+class UAssetDefinition_AudioSynesthesiaNRTSettings : public UAssetDefinitionDefault_AudioDiffable
+{
+	GENERATED_BODY()
+public:
+	// UAssetDefinition Begin
+	virtual FText GetAssetDisplayName() const override { return NSLOCTEXT("AssetDefinition", "AssetDefinition_AudioSynesthesiaNRTSettings", "Synesthesia NRT Settings"); }
+	virtual FLinearColor GetAssetColor() const override { return FColor(100.0f, 100.0f, 100.0f); }
+	virtual TSoftClassPtr<UObject> GetAssetClass() const override { return UAudioSynesthesiaNRTSettings::StaticClass(); }
+	virtual TConstArrayView<FAssetCategoryPath> GetAssetCategories() const override
+	{
+		static const auto Categories = 
+		{
+			FAssetCategoryPath(EAssetCategoryPaths::Audio,
+				NSLOCTEXT("AssetDefinition", "AssetDefinition_AudioSynesthesiaNRTSettingsSubMenu", "Advanced"),
+				FCategoryPath(NSLOCTEXT("AssetDefinition", "AssetDefinition_AudioSynesthesiaNRTSettingsSubMenuSection", "Analysis"), ECategoryMenuType::Section))
+		};
+		return Categories;
+	}
+	// UAssetDefinition End
+};
diff --git a/Engine/Plugins/Runtime/AudioSynesthesia/Source/AudioSynesthesiaEditor/Private/AssetDefinition_AudioSynesthesiaSettings.h b/Engine/Plugins/Runtime/AudioSynesthesia/Source/AudioSynesthesiaEditor/Private/AssetDefinition_AudioSynesthesiaSettings.h
new file mode 100644
--- /dev/null
+++ b/Engine/Plugins/Runtime/AudioSynesthesia/Source/AudioSynesthesiaEditor/Private/AssetDefinition_AudioSynesthesiaSettings.h
@@ -0,0 +1,28 @@
+// Copyright Epic Games, Inc. All Rights Reserved.
+
+#pragma once
+#include "AssetTypeActions/AssetDefinitionDefault_AudioDiffable.h"
+#include "AudioSynesthesia.h"
+#include "AssetDefinition_AudioSynesthesiaSettings.generated.h"
+
+UCLASS()
+class UAssetDefinition_AudioSynesthesiaSettings : public UAssetDefinitionDefault_AudioDiffable
+{
+	GENERATED_BODY()
+public:
+	// UAssetDefinition Begin
+	virtual FText GetAssetDisplayName() const override { return NSLOCTEXT("AssetDefinition", "AssetDefinition_AudioSynesthesiaSettings", "Synesthesia Settings"); }
+	virtual FLinearColor GetAssetColor() const override { return FColor(100.0f, 50.0f, 100.0f); }
+	virtual TSoftClassPtr<UObject> GetAssetClass() const override { return UAudioSynesthesiaSettings::StaticClass(); }
+	virtual TConstArrayView<FAssetCategoryPath> GetAssetCategories() const override
+	{
+		static const auto Categories = 
+		{
+			FAssetCategoryPath(EAssetCategoryPaths::Audio,
+				NSLOCTEXT("AssetDefinition", "AssetDefinition_AudioSynesthesiaSettingsSubMenu", "Advanced"),
+				FCategoryPath(NSLOCTEXT("AssetDefinition", "AssetDefinition_AudioSynesthesiaSettingsSubMenuSection", "Analysis"), ECategoryMenuType::Section))
+		};
+		return Categories;
+	}
+	// UAssetDefinition End
+};
diff --git a/Engine/Plugins/Runtime/AudioSynesthesia/Source/AudioSynesthesiaEditor/Private/AssetTypeActions_AudioSynesthesiaNRT.cpp b/Engine/Plugins/Runtime/AudioSynesthesia/Source/AudioSynesthesiaEditor/Private/AssetTypeActions_AudioSynesthesiaNRT.cpp
deleted file mode 100644
diff --git a/Engine/Plugins/Runtime/AudioSynesthesia/Source/AudioSynesthesiaEditor/Private/AssetTypeActions_AudioSynesthesiaNRT.h b/Engine/Plugins/Runtime/AudioSynesthesia/Source/AudioSynesthesiaEditor/Private/AssetTypeActions_AudioSynesthesiaNRT.h
deleted file mode 100644
diff --git a/Engine/Plugins/Runtime/AudioSynesthesia/Source/AudioSynesthesiaEditor/Private/AssetTypeActions_AudioSynesthesiaNRTSettings.cpp b/Engine/Plugins/Runtime/AudioSynesthesia/Source/AudioSynesthesiaEditor/Private/AssetTypeActions_AudioSynesthesiaNRTSettings.cpp
deleted file mode 100644
diff --git a/Engine/Plugins/Runtime/AudioSynesthesia/Source/AudioSynesthesiaEditor/Private/AssetTypeActions_AudioSynesthesiaNRTSettings.h b/Engine/Plugins/Runtime/AudioSynesthesia/Source/AudioSynesthesiaEditor/Private/AssetTypeActions_AudioSynesthesiaNRTSettings.h
deleted file mode 100644
diff --git a/Engine/Plugins/Runtime/AudioSynesthesia/Source/AudioSynesthesiaEditor/Private/AssetTypeActions_AudioSynesthesiaSettings.cpp b/Engine/Plugins/Runtime/AudioSynesthesia/Source/AudioSynesthesiaEditor/Private/AssetTypeActions_AudioSynesthesiaSettings.cpp
deleted file mode 100644
diff --git a/Engine/Plugins/Runtime/AudioSynesthesia/Source/AudioSynesthesiaEditor/Private/AssetTypeActions_AudioSynesthesiaSettings.h b/Engine/Plugins/Runtime/AudioSynesthesia/Source/AudioSynesthesiaEditor/Private/AssetTypeActions_AudioSynesthesiaSettings.h
deleted file mode 100644
diff --git a/Engine/Plugins/Runtime/AudioSynesthesia/Source/AudioSynesthesiaEditor/Private/AudioSynesthesiaEditorModule.cpp b/Engine/Plugins/Runtime/AudioSynesthesia/Source/AudioSynesthesiaEditor/Private/AudioSynesthesiaEditorModule.cpp
--- a/Engine/Plugins/Runtime/AudioSynesthesia/Source/AudioSynesthesiaEditor/Private/AudioSynesthesiaEditorModule.cpp
+++ b/Engine/Plugins/Runtime/AudioSynesthesia/Source/AudioSynesthesiaEditor/Private/AudioSynesthesiaEditorModule.cpp
@@ -2,9 +2,8 @@
 
 #include "AudioSynesthesiaEditorModule.h"
 
-#include "AssetTypeActions_AudioSynesthesiaNRT.h"
-#include "AssetTypeActions_AudioSynesthesiaNRTSettings.h"
-#include "AssetTypeActions_AudioSynesthesiaSettings.h"
+#include "AssetDefinitionRegistry.h"
+#include "AssetDefinition_AudioSynesthesiaDynamic.h"
 #include "AudioAnalyzerModule.h"
 #include "AudioSynesthesia.h"
 #include "AudioSynesthesiaNRT.h"
@@ -31,22 +30,16 @@ public:
 
 	virtual void RegisterAssetActions() override
 	{
-		RegisterAudioSynesthesiaAssetActions<UAudioSynesthesiaNRT, FAssetTypeActions_AudioSynesthesiaNRT>();
-		RegisterAudioSynesthesiaAssetActions<UAudioSynesthesiaNRTSettings, FAssetTypeActions_AudioSynesthesiaNRTSettings>();
-		RegisterAudioSynesthesiaAssetActions<UAudioSynesthesiaSettings, FAssetTypeActions_AudioSynesthesiaSettings>();
+		RegisterAudioSynesthesiaChildAssetDefinitionFor<UAudioSynesthesiaNRT>();
+		RegisterAudioSynesthesiaChildAssetDefinitionFor<UAudioSynesthesiaNRTSettings>();
+		RegisterAudioSynesthesiaChildAssetDefinitionFor<UAudioSynesthesiaSettings>();
 	}
 
 private:
 
-	template<class SynesthesiaAssetType, class SynesthesiaAssetActionsType >
-	void RegisterAudioSynesthesiaAssetActions() 
+	template<class SynesthesiaAssetType>
+	void RegisterAudioSynesthesiaChildAssetDefinitionFor() 
 	{
-		// Register the audio editor asset type actions
-		IAssetTools& AssetTools = FModuleManager::LoadModuleChecked<FAssetToolsModule>("AssetTools").Get();
-
-		// Register base abstract class for asset filtering
-		AssetTools.RegisterAssetTypeActions(MakeShared<SynesthesiaAssetActionsType>(nullptr));
-
 		// Look for any sound effect presets to register
 		for (TObjectIterator<UClass> It; It; ++It)
 		{
@@ -66,12 +59,16 @@ private:
 				if (!RegisteredActions.Contains(Synesthesia) && Synesthesia->HasAssetActions())
 				{
 					RegisteredActions.Add(Synesthesia);
-					AssetTools.RegisterAssetTypeActions(MakeShared<SynesthesiaAssetActionsType>(Synesthesia));
+					UAssetDefinition_AudioSynesthesiaDynamic* Definition = NewObject<UAssetDefinition_AudioSynesthesiaDynamic>(
+						GetTransientPackage(),
+						MakeUniqueObjectName(GetTransientPackage(), UAssetDefinition_AudioSynesthesiaDynamic::StaticClass())
+					);
+					Definition->Initialize(TSubclassOf<UAudioAnalyzerAssetBase>(ChildClass));
+					UAssetDefinitionRegistry::Get()->RegisterAssetDefinition(Definition);
 				}
 			}
 		}
 	}	
-
 	TSet<UObject*> RegisteredActions;
 };
 
diff --git a/Engine/Plugins/Runtime/Harmonix/Source/HarmonixMetasoundEditor/Private/AssetDefinition_MidiStepSequence.cpp b/Engine/Plugins/Runtime/Harmonix/Source/HarmonixMetasoundEditor/Private/AssetDefinition_MidiStepSequence.cpp
--- a/Engine/Plugins/Runtime/Harmonix/Source/HarmonixMetasoundEditor/Private/AssetDefinition_MidiStepSequence.cpp
+++ b/Engine/Plugins/Runtime/Harmonix/Source/HarmonixMetasoundEditor/Private/AssetDefinition_MidiStepSequence.cpp
@@ -5,6 +5,8 @@
 
 #include UE_INLINE_GENERATED_CPP_BY_NAME(AssetDefinition_MidiStepSequence)
 
+#define LOCTEXT_NAMESPACE "AssetDefinition_MidiStepSequence"
+
 TSoftClassPtr<UObject> UAssetDefinition_MidiStepSequence::GetAssetClass() const
 {
 	return UMidiStepSequence::StaticClass();
@@ -12,7 +14,7 @@ TSoftClassPtr<UObject> UAssetDefinition_MidiStepSequence::GetAssetClass() const
 
 FText UAssetDefinition_MidiStepSequence::GetAssetDisplayName() const
 {
-	return NSLOCTEXT("AssetTypeActions", "MIDIStepSequenceDefinition", "MIDI Step Sequence");
+	return LOCTEXT("MIDIStepSequenceDefinition", "MIDI Step Sequence");
 }
 
 FLinearColor  UAssetDefinition_MidiStepSequence::GetAssetColor() const
@@ -22,7 +24,12 @@ FLinearColor  UAssetDefinition_MidiStepSequence::GetAssetColor() const
 
 TConstArrayView<FAssetCategoryPath> UAssetDefinition_MidiStepSequence::GetAssetCategories() const
 {
-	static const auto Categories = { EAssetCategoryPaths::Audio / NSLOCTEXT("Harmonix", "HmxAssetCategoryName", "Harmonix") };
+	static const auto Categories = 
+		{
+			FAssetCategoryPath(EAssetCategoryPaths::Audio, 
+				LOCTEXT("AssetDefinition_MidiStepSequenceSubMenu", "Advanced"), 
+				FCategoryPath(LOCTEXT("AssetDefinition_MidiStepSequenceSubMenuSection", "Harmonix"), ECategoryMenuType::Section))
+		};
 	return Categories;
 }
 
@@ -31,3 +38,4 @@ bool UAssetDefinition_MidiStepSequence::CanImport() const
 	return true;
 }
 
+#undef LOCTEXT_NAMESPACE
diff --git a/Engine/Plugins/Runtime/Harmonix/Source/HarmonixMetasoundEditor/Private/AssetDefinition_MidiStutterSequence.cpp b/Engine/Plugins/Runtime/Harmonix/Source/HarmonixMetasoundEditor/Private/AssetDefinition_MidiStutterSequence.cpp
--- a/Engine/Plugins/Runtime/Harmonix/Source/HarmonixMetasoundEditor/Private/AssetDefinition_MidiStutterSequence.cpp
+++ b/Engine/Plugins/Runtime/Harmonix/Source/HarmonixMetasoundEditor/Private/AssetDefinition_MidiStutterSequence.cpp
@@ -5,6 +5,8 @@
 
 #include UE_INLINE_GENERATED_CPP_BY_NAME(AssetDefinition_MidiStutterSequence)
 
+#define LOCTEXT_NAMESPACE "AssetDefinition_MidiStutterSequence"
+
 TSoftClassPtr<UObject> UAssetDefinition_MidiStutterSequence::GetAssetClass() const
 {
 	return UMidiStutterSequence::StaticClass();
@@ -12,7 +14,7 @@ TSoftClassPtr<UObject> UAssetDefinition_MidiStutterSequence::GetAssetClass() con
 
 FText UAssetDefinition_MidiStutterSequence::GetAssetDisplayName() const
 {
-	return NSLOCTEXT("AssetTypeActions", "MIDIStutterSequenceDefinition", "MIDI Stutter Sequence");
+	return LOCTEXT("MIDIStutterSequenceDefinition", "MIDI Stutter Sequence");
 }
 
 FLinearColor  UAssetDefinition_MidiStutterSequence::GetAssetColor() const
@@ -22,7 +24,12 @@ FLinearColor  UAssetDefinition_MidiStutterSequence::GetAssetColor() const
 
 TConstArrayView<FAssetCategoryPath> UAssetDefinition_MidiStutterSequence::GetAssetCategories() const
 {
-	static const auto Categories = { EAssetCategoryPaths::Audio / NSLOCTEXT("Harmonix", "HmxAssetCategoryName", "Harmonix") };
+	static const auto Categories = 
+		{
+			FAssetCategoryPath(EAssetCategoryPaths::Audio, 
+				LOCTEXT("AssetDefinition_MidiStutterSequenceSubMenu", "Advanced"), 
+				FCategoryPath(LOCTEXT("AssetDefinition_MidiStutterSequenceSubMenuSection", "Harmonix"), ECategoryMenuType::Section))
+		};
 	return Categories;
 }
 
@@ -31,3 +38,4 @@ bool UAssetDefinition_MidiStutterSequence::CanImport() const
 	return false;
 }
 
+#undef LOCTEXT_NAMESPACE
diff --git a/Engine/Plugins/Runtime/Harmonix/Source/HarmonixMetasoundEditor/Private/AssetDefinition_MusicAssets.cpp b/Engine/Plugins/Runtime/Harmonix/Source/HarmonixMetasoundEditor/Private/AssetDefinition_MusicAssets.cpp
--- a/Engine/Plugins/Runtime/Harmonix/Source/HarmonixMetasoundEditor/Private/AssetDefinition_MusicAssets.cpp
+++ b/Engine/Plugins/Runtime/Harmonix/Source/HarmonixMetasoundEditor/Private/AssetDefinition_MusicAssets.cpp
@@ -7,6 +7,8 @@
 
 #include UE_INLINE_GENERATED_CPP_BY_NAME(AssetDefinition_MusicAssets)
 
+#define LOCTEXT_NAMESPACE "AssetDefinition_MusicAssets"
+
 TSoftClassPtr<UObject> UAssetDefinition_MetasoundMusic::GetAssetClass() const
 {
 	return UHarmonixMetasoundMusicAsset::StaticClass();
@@ -24,7 +26,12 @@ FLinearColor  UAssetDefinition_MetasoundMusic::GetAssetColor() const
 
 TConstArrayView<FAssetCategoryPath> UAssetDefinition_MetasoundMusic::GetAssetCategories() const
 {
-	static const auto Categories = { EAssetCategoryPaths::Audio / NSLOCTEXT("Harmonix", "HmxAssetCategoryName", "Harmonix") };
+	static const auto Categories = 
+		{
+			FAssetCategoryPath(EAssetCategoryPaths::Audio, 
+				LOCTEXT("AssetDefinition_MetasoundMusicSubMenu", "Advanced"), 
+				FCategoryPath(LOCTEXT("AssetDefinition_MetasoundMusicSubMenuSection", "Harmonix"), ECategoryMenuType::Section))
+		};
 	return Categories;
 }
 
@@ -60,7 +67,12 @@ FLinearColor  UAssetDefinition_WaveMusic::GetAssetColor() const
 
 TConstArrayView<FAssetCategoryPath> UAssetDefinition_WaveMusic::GetAssetCategories() const
 {
-	static const auto Categories = { EAssetCategoryPaths::Audio / NSLOCTEXT("Harmonix", "HmxAssetCategoryName", "Harmonix") };
+	static const auto Categories = 
+		{
+			FAssetCategoryPath(EAssetCategoryPaths::Audio, 
+				LOCTEXT("AssetDefinition_WaveMusicSubMenu", "Advanced"), 
+				FCategoryPath(LOCTEXT("AssetDefinition_WaveMusicSubMenuSection", "Harmonix"), ECategoryMenuType::Section))
+		};
 	return Categories;
 }
 
@@ -78,3 +90,5 @@ const FSlateBrush* UAssetDefinition_WaveMusic::GetThumbnailBrush(const FAssetDat
 {
 	return HarmonixMetasoundEditor::FSlateStyle::Get().GetBrush("HarmonixMetasoundEditor.WaveMusic.Thumbnail");
 }
+
+#undef LOCTEXT_NAMESPACE
diff --git a/Engine/Plugins/Runtime/Metasound/Source/MetasoundEditor/Private/MetasoundAssetDefinitions.cpp b/Engine/Plugins/Runtime/Metasound/Source/MetasoundEditor/Private/MetasoundAssetDefinitions.cpp
--- a/Engine/Plugins/Runtime/Metasound/Source/MetasoundEditor/Private/MetasoundAssetDefinitions.cpp
+++ b/Engine/Plugins/Runtime/Metasound/Source/MetasoundEditor/Private/MetasoundAssetDefinitions.cpp
@@ -576,7 +576,12 @@ TSoftClassPtr<UObject> UAssetDefinition_MetaSoundPatch::GetAssetClass() const
 TConstArrayView<FAssetCategoryPath> UAssetDefinition_MetaSoundPatch::GetAssetCategories() const
 {
 	static const auto Pinned_Categories = { EAssetCategoryPaths::Audio };
-	static const auto Categories = { EAssetCategoryPaths::Audio / LOCTEXT("AssetSoundMetaSoundsSubMenu", "MetaSounds") };
+	static const auto Categories = 
+		{
+			FAssetCategoryPath(EAssetCategoryPaths::Audio, 
+				LOCTEXT("AssetSoundMetaSoundsSubMenu", "Advanced"), 
+				FCategoryPath(LOCTEXT("AssetSoundMetaSoundsSubMenuSection", "MetaSounds"), ECategoryMenuType::Section))
+		};
 
 	if (GetDefault<UMetasoundEditorSettings>()->bPinMetaSoundPatchInAssetMenu)
 	{
diff --git a/Engine/Plugins/Runtime/ResonanceAudio/Source/ResonanceAudioEditor/Private/AssetDefinition_ResonanceAudioReverbPluginPreset.h b/Engine/Plugins/Runtime/ResonanceAudio/Source/ResonanceAudioEditor/Private/AssetDefinition_ResonanceAudioReverbPluginPreset.h
new file mode 100644
--- /dev/null
+++ b/Engine/Plugins/Runtime/ResonanceAudio/Source/ResonanceAudioEditor/Private/AssetDefinition_ResonanceAudioReverbPluginPreset.h
@@ -0,0 +1,34 @@
+//
+// Copyright (C) Google Inc. 2017. All rights reserved.
+//
+
+#pragma once
+#include "AssetDefinitionDefault.h"
+#include "ResonanceAudioConstants.h"
+#include "ResonanceAudioReverb.h"
+#include "AssetDefinition_ResonanceAudioReverbPluginPreset.generated.h"
+
+UCLASS()
+class UAssetDefinition_ResonanceAudioReverbPluginPreset : public UAssetDefinitionDefault
+{
+	GENERATED_BODY()
+
+public:
+	// UAssetDefinition Begin
+	virtual FText GetAssetDisplayName() const override { return NSLOCTEXT("AssetDefinition", "AssetDefinition_ResonanceAudioReverbPluginPreset", "Resonance Audio Reverb Settings"); }
+	virtual FLinearColor GetAssetColor() const override { return ResonanceAudio::ASSET_COLOR; }
+	virtual TSoftClassPtr<UObject> GetAssetClass() const override { return UResonanceAudioReverbPluginPreset::StaticClass(); }
+	virtual bool CanRegisterStatically() const override { return true; }
+	virtual TConstArrayView<FAssetCategoryPath> GetAssetCategories() const override
+	{
+		static const auto Categories = 
+		{
+			FAssetCategoryPath(EAssetCategoryPaths::Audio,
+				NSLOCTEXT("AssetDefinition", "AssetDefinition_ResonanceAudioReverbPluginPresetSubMenu", "Advanced"),
+				FCategoryPath(NSLOCTEXT("AssetDefinition", "AssetDefinition_ResonanceAudioReverbPluginPresetSubMenuSection", "Binaural Spatialization"), ECategoryMenuType::Section))
+		};
+
+		return Categories;
+	}
+	// UAssetDefinition End
+};
diff --git a/Engine/Plugins/Runtime/ResonanceAudio/Source/ResonanceAudioEditor/Private/AssetDefinition_ResonanceAudioSpatializationSourceSettings.h b/Engine/Plugins/Runtime/ResonanceAudio/Source/ResonanceAudioEditor/Private/AssetDefinition_ResonanceAudioSpatializationSourceSettings.h
new file mode 100644
--- /dev/null
+++ b/Engine/Plugins/Runtime/ResonanceAudio/Source/ResonanceAudioEditor/Private/AssetDefinition_ResonanceAudioSpatializationSourceSettings.h
@@ -0,0 +1,33 @@
+//
+// Copyright (C) Google Inc. 2017. All rights reserved.
+//
+
+#pragma once
+#include "AssetDefinitionDefault.h"
+#include "ResonanceAudioConstants.h"
+#include "ResonanceAudioSpatializationSourceSettings.h"
+#include "AssetDefinition_ResonanceAudioSpatializationSourceSettings.generated.h"
+
+UCLASS()
+class UAssetDefinition_ResonanceAudioSpatializationSourceSettings : public UAssetDefinitionDefault
+{
+	GENERATED_BODY()
+
+public:
+	// UAssetDefinition Begin
+	virtual FText GetAssetDisplayName() const override { return NSLOCTEXT("AssetDefinition", "AssetDefinition_ResonanceAudioSpatializationSourceSettings", "Resonance Audio Spatialization Settings"); }
+	virtual FLinearColor GetAssetColor() const override { return ResonanceAudio::ASSET_COLOR; }
+	virtual TSoftClassPtr<UObject> GetAssetClass() const override { return UResonanceAudioSpatializationSourceSettings::StaticClass(); }
+	virtual TConstArrayView<FAssetCategoryPath> GetAssetCategories() const override
+	{
+		static const auto Categories = 
+		{
+			FAssetCategoryPath(EAssetCategoryPaths::Audio,
+				NSLOCTEXT("AssetDefinition", "AssetDefinition_ResonanceAudioSpatializationSourceSettingsSubMenu", "Advanced"),
+				FCategoryPath(NSLOCTEXT("AssetDefinition", "AssetDefinition_ResonanceAudioSpatializationSourceSettingsSubMenuSection", "Binaural Spatialization"), ECategoryMenuType::Section))
+		};
+
+		return Categories;
+	}
+	// UAssetDefinition End
+};
diff --git a/Engine/Plugins/Runtime/ResonanceAudio/Source/ResonanceAudioEditor/Private/ResonanceAudioEditorModule.cpp b/Engine/Plugins/Runtime/ResonanceAudio/Source/ResonanceAudioEditor/Private/ResonanceAudioEditorModule.cpp
--- a/Engine/Plugins/Runtime/ResonanceAudio/Source/ResonanceAudioEditor/Private/ResonanceAudioEditorModule.cpp
+++ b/Engine/Plugins/Runtime/ResonanceAudio/Source/ResonanceAudioEditor/Private/ResonanceAudioEditorModule.cpp
@@ -48,12 +48,6 @@ namespace ResonanceAudio
 		ResonanceAudioStyleSet->Set("ClassThumbnail.ResonanceAudioReverbPluginPreset", new FSlateImageBrush(ResonanceAudioContent + "/ResonanceAudioReverbPluginPreset_64.png", Vec64));
 
 		FSlateStyleRegistry::RegisterSlateStyle(*ResonanceAudioStyleSet.Get());
-
-		// Register the audio editor asset type actions.
-		IAssetTools& AssetTools = FModuleManager::LoadModuleChecked<FAssetToolsModule>("AssetTools").Get();
-
-		AssetTools.RegisterAssetTypeActions(MakeShared<FAssetTypeActions_ResonanceAudioReverbPluginPreset>());
-		AssetTools.RegisterAssetTypeActions(MakeShared<FAssetTypeActions_ResonanceAudioSpatializationSourceSettings>());
 	}
 
 	void FResonanceAudioEditorModule::ShutdownModule()
diff --git a/Engine/Plugins/Runtime/ResonanceAudio/Source/ResonanceAudioEditor/Private/ResonanceAudioReverbPluginPresetFactory.cpp b/Engine/Plugins/Runtime/ResonanceAudio/Source/ResonanceAudioEditor/Private/ResonanceAudioReverbPluginPresetFactory.cpp
--- a/Engine/Plugins/Runtime/ResonanceAudio/Source/ResonanceAudioEditor/Private/ResonanceAudioReverbPluginPresetFactory.cpp
+++ b/Engine/Plugins/Runtime/ResonanceAudio/Source/ResonanceAudioEditor/Private/ResonanceAudioReverbPluginPresetFactory.cpp
@@ -3,43 +3,13 @@
 //
 
 #include "ResonanceAudioReverbPluginPresetFactory.h"
-#include "ResonanceAudioConstants.h"
+#include "AssetTypeCategories.h"
 #include "ResonanceAudioReverb.h"
 
 #include "AudioAnalytics.h"
 
 #include UE_INLINE_GENERATED_CPP_BY_NAME(ResonanceAudioReverbPluginPresetFactory)
 
-FText FAssetTypeActions_ResonanceAudioReverbPluginPreset::GetName() const
-{
-	return NSLOCTEXT("AssetTypeActions", "AssetTypeActions_ResonanceAudioReverbPluginPreset", "Resonance Audio Reverb Settings");
-}
-
-FColor FAssetTypeActions_ResonanceAudioReverbPluginPreset::GetTypeColor() const
-{
-	return ResonanceAudio::ASSET_COLOR;
-}
-
-UClass* FAssetTypeActions_ResonanceAudioReverbPluginPreset::GetSupportedClass() const
-{
-	return UResonanceAudioReverbPluginPreset::StaticClass();
-}
-
-uint32 FAssetTypeActions_ResonanceAudioReverbPluginPreset::GetCategories()
-{
-	return EAssetTypeCategories::Sounds;
-}
-
-const TArray<FText>& FAssetTypeActions_ResonanceAudioReverbPluginPreset::GetSubMenus() const
-{
-	static const TArray<FText> SubMenus
-	{
-		NSLOCTEXT("AssetTypeActions", "AssetTypeActions_AssetSoundResonanceSubMenu", "Resonance")
-	};
-
-	return SubMenus;
-}
-
 UResonanceAudioReverbPluginPresetFactory::UResonanceAudioReverbPluginPresetFactory(const FObjectInitializer& ObjectInitializer)
 	: Super(ObjectInitializer)
 {
diff --git a/Engine/Plugins/Runtime/ResonanceAudio/Source/ResonanceAudioEditor/Private/ResonanceAudioReverbPluginPresetFactory.h b/Engine/Plugins/Runtime/ResonanceAudio/Source/ResonanceAudioEditor/Private/ResonanceAudioReverbPluginPresetFactory.h
--- a/Engine/Plugins/Runtime/ResonanceAudio/Source/ResonanceAudioEditor/Private/ResonanceAudioReverbPluginPresetFactory.h
+++ b/Engine/Plugins/Runtime/ResonanceAudio/Source/ResonanceAudioEditor/Private/ResonanceAudioReverbPluginPresetFactory.h
@@ -5,19 +5,8 @@
 #pragma once
 
 #include "Factories/Factory.h"
-#include "AssetTypeActions_Base.h"
 #include "ResonanceAudioReverbPluginPresetFactory.generated.h"
 
-class FAssetTypeActions_ResonanceAudioReverbPluginPreset : public FAssetTypeActions_Base
-{
-public:
-	virtual FText GetName() const override;
-	virtual FColor GetTypeColor() const override;
-	virtual UClass* GetSupportedClass() const override;
-	virtual uint32 GetCategories() override;
-	virtual const TArray<FText>& GetSubMenus() const override;
-};
-
 UCLASS(MinimalAPI, hidecategories = Object)
 class UResonanceAudioReverbPluginPresetFactory : public UFactory
 {
diff --git a/Engine/Plugins/Runtime/ResonanceAudio/Source/ResonanceAudioEditor/Private/ResonanceAudioSpatializationSourceSettingsFactory.cpp b/Engine/Plugins/Runtime/ResonanceAudio/Source/ResonanceAudioEditor/Private/ResonanceAudioSpatializationSourceSettingsFactory.cpp
--- a/Engine/Plugins/Runtime/ResonanceAudio/Source/ResonanceAudioEditor/Private/ResonanceAudioSpatializationSourceSettingsFactory.cpp
+++ b/Engine/Plugins/Runtime/ResonanceAudio/Source/ResonanceAudioEditor/Private/ResonanceAudioSpatializationSourceSettingsFactory.cpp
@@ -3,43 +3,13 @@
 //
 
 #include "ResonanceAudioSpatializationSourceSettingsFactory.h"
+#include "AssetTypeCategories.h"
 #include "ResonanceAudioSpatializationSourceSettings.h"
-#include "ResonanceAudioConstants.h"
 
 #include "AudioAnalytics.h"
 
 #include UE_INLINE_GENERATED_CPP_BY_NAME(ResonanceAudioSpatializationSourceSettingsFactory)
 
-FText FAssetTypeActions_ResonanceAudioSpatializationSourceSettings::GetName() const
-{
-	return NSLOCTEXT("AssetTypeActions", "AssetTypeActions_ResonanceAudioSpatializationSourceSettings", "Resonance Audio Spatialization Settings");
-}
-
-FColor FAssetTypeActions_ResonanceAudioSpatializationSourceSettings::GetTypeColor() const
-{
-	return ResonanceAudio::ASSET_COLOR;
-}
-
-UClass* FAssetTypeActions_ResonanceAudioSpatializationSourceSettings::GetSupportedClass() const
-{
-	return UResonanceAudioSpatializationSourceSettings::StaticClass();
-}
-
-uint32 FAssetTypeActions_ResonanceAudioSpatializationSourceSettings::GetCategories()
-{
-	return EAssetTypeCategories::Sounds;
-}
-
-const TArray<FText>& FAssetTypeActions_ResonanceAudioSpatializationSourceSettings::GetSubMenus() const
-{
-	static const TArray<FText> SubMenus
-	{
-		NSLOCTEXT("AssetTypeActions", "AssetTypeActions_AssetSoundResonanceSubMenu", "Resonance")
-	};
-
-	return SubMenus;
-}
-
 UResonanceAudioSpatializationSourceSettingsFactory::UResonanceAudioSpatializationSourceSettingsFactory(const FObjectInitializer& ObjectInitializer)
 	: Super(ObjectInitializer)
 {
diff --git a/Engine/Plugins/Runtime/ResonanceAudio/Source/ResonanceAudioEditor/Private/ResonanceAudioSpatializationSourceSettingsFactory.h b/Engine/Plugins/Runtime/ResonanceAudio/Source/ResonanceAudioEditor/Private/ResonanceAudioSpatializationSourceSettingsFactory.h
--- a/Engine/Plugins/Runtime/ResonanceAudio/Source/ResonanceAudioEditor/Private/ResonanceAudioSpatializationSourceSettingsFactory.h
+++ b/Engine/Plugins/Runtime/ResonanceAudio/Source/ResonanceAudioEditor/Private/ResonanceAudioSpatializationSourceSettingsFactory.h
@@ -5,19 +5,8 @@
 #pragma once
 
 #include "Factories/Factory.h"
-#include "AssetTypeActions_Base.h"
 #include "ResonanceAudioSpatializationSourceSettingsFactory.generated.h"
 
-class FAssetTypeActions_ResonanceAudioSpatializationSourceSettings : public FAssetTypeActions_Base
-{
-public:
-	virtual FText GetName() const override;
-	virtual FColor GetTypeColor() const override;
-	virtual UClass* GetSupportedClass() const override;
-	virtual uint32 GetCategories() override;
-	virtual const TArray<FText>& GetSubMenus() const override;
-};
-
 UCLASS(MinimalAPI, hidecategories = Object)
 class UResonanceAudioSpatializationSourceSettingsFactory : public UFactory
 {
diff --git a/Engine/Plugins/Runtime/ResonanceAudio/Source/ResonanceAudioEditor/ResonanceAudioEditor.Build.cs b/Engine/Plugins/Runtime/ResonanceAudio/Source/ResonanceAudioEditor/ResonanceAudioEditor.Build.cs
--- a/Engine/Plugins/Runtime/ResonanceAudio/Source/ResonanceAudioEditor/ResonanceAudioEditor.Build.cs
+++ b/Engine/Plugins/Runtime/ResonanceAudio/Source/ResonanceAudioEditor/ResonanceAudioEditor.Build.cs
@@ -35,13 +35,9 @@ namespace UnrealBuildTool.Rules
                 }
             );
 
-            PrivateIncludePathModuleNames.AddRange(
-                new string[] {
-                    "AssetTools",
-            });
-
             PrivateDependencyModuleNames.AddRange(
                 new string[] {
+	                "AssetDefinition",
                     "Slate",
                     "SlateCore",
 					"EditorFramework",
diff --git a/Engine/Plugins/Runtime/SoundCueTemplates/Source/SoundCueTemplatesEditor/Private/AssetDefinition_SoundCueTemplate.cpp b/Engine/Plugins/Runtime/SoundCueTemplates/Source/SoundCueTemplatesEditor/Private/AssetDefinition_SoundCueTemplate.cpp
--- a/Engine/Plugins/Runtime/SoundCueTemplates/Source/SoundCueTemplatesEditor/Private/AssetDefinition_SoundCueTemplate.cpp
+++ b/Engine/Plugins/Runtime/SoundCueTemplates/Source/SoundCueTemplatesEditor/Private/AssetDefinition_SoundCueTemplate.cpp
@@ -25,7 +25,7 @@
 TConstArrayView<FAssetCategoryPath> UAssetDefinition_SoundCueTemplate::GetAssetCategories() const
 {
 	static const auto Pinned_Categories = { EAssetCategoryPaths::Audio };
-	static const auto Categories = { EAssetCategoryPaths::Audio / LOCTEXT("AssetSoundCueSubMenu", "Source") };
+	static const auto Categories = { FAssetCategoryPath(EAssetCategoryPaths::Audio, LOCTEXT("AssetSoundCue_SubMenu", "Legacy")) };
 	
 	if (GetDefault<UAudioEditorSettings>()->bPinSoundCueTemplateInAssetMenu)
 	{
diff --git a/Engine/Plugins/Runtime/Soundscape/Source/SoundScapeEditor/Private/AssetDefinition_SoundScapeColor.h b/Engine/Plugins/Runtime/Soundscape/Source/SoundScapeEditor/Private/AssetDefinition_SoundScapeColor.h
new file mode 100644
--- /dev/null
+++ b/Engine/Plugins/Runtime/Soundscape/Source/SoundScapeEditor/Private/AssetDefinition_SoundScapeColor.h
@@ -0,0 +1,28 @@
+﻿// Copyright Epic Games, Inc. All Rights Reserved.
+
+#pragma once
+#include "AssetDefinitionDefault.h"
+#include "SoundscapeColor.h"
+#include "AssetDefinition_SoundScapeColor.generated.h"
+
+UCLASS()
+class UAssetDefinition_SoundScapeColor : public UAssetDefinitionDefault
+{
+	GENERATED_BODY()
+public:
+	// UAssetDefinition Begin
+	virtual FText GetAssetDisplayName() const override { return NSLOCTEXT("AssetDefinition", "AssetDefinition_SoundscapeColor", "Soundscape Color"); }
+	virtual FLinearColor GetAssetColor() const override { return FColor(0, 175, 175); }
+	virtual TSoftClassPtr<UObject> GetAssetClass() const override { return USoundscapeColor::StaticClass(); }
+	virtual TConstArrayView<FAssetCategoryPath> GetAssetCategories() const override
+	{
+		static const auto Categories = 
+			{
+				FAssetCategoryPath(EAssetCategoryPaths::Audio,
+					NSLOCTEXT("AssetDefinition", "AssetDefinition_SoundscapeColorSubMenu", "Advanced"),
+					FCategoryPath(NSLOCTEXT("AssetDefinition", "AssetDefinition_SoundscapeColorSubMenuSection", "SoundScape"), ECategoryMenuType::Section))
+			};
+		return Categories;
+	}
+	// UAssetDefinition End
+};
diff --git a/Engine/Plugins/Runtime/Soundscape/Source/SoundScapeEditor/Private/AssetDefinition_SoundScapePalette.h b/Engine/Plugins/Runtime/Soundscape/Source/SoundScapeEditor/Private/AssetDefinition_SoundScapePalette.h
new file mode 100644
--- /dev/null
+++ b/Engine/Plugins/Runtime/Soundscape/Source/SoundScapeEditor/Private/AssetDefinition_SoundScapePalette.h
@@ -0,0 +1,28 @@
+﻿// Copyright Epic Games, Inc. All Rights Reserved.
+
+#pragma once
+#include "AssetDefinitionDefault.h"
+#include "SoundScapePalette.h"
+#include "AssetDefinition_SoundScapePalette.generated.h"
+
+UCLASS()
+class UAssetDefinition_SoundScapePalette : public UAssetDefinitionDefault
+{
+	GENERATED_BODY()
+public:
+	// UAssetDefinition Begin
+	virtual FText GetAssetDisplayName() const override { return NSLOCTEXT("AssetDefinition", "AssetDefinition_SoundScapePalette", "Soundscape Palette"); }
+	virtual FLinearColor GetAssetColor() const override { return FColor(0, 125, 255); }
+	virtual TSoftClassPtr<UObject> GetAssetClass() const override { return USoundscapePalette::StaticClass(); }
+	virtual TConstArrayView<FAssetCategoryPath> GetAssetCategories() const override
+	{
+		static const auto Categories = 
+		{
+			FAssetCategoryPath(EAssetCategoryPaths::Audio,
+				NSLOCTEXT("AssetDefinition", "AssetDefinition_SoundScapePaletteSubMenu", "Advanced"),
+				FCategoryPath(NSLOCTEXT("AssetDefinition", "AssetDefinition_SoundScapePaletteSubMenuSection", "SoundScape"), ECategoryMenuType::Section))
+		};
+		return Categories;
+	}
+	// UAssetDefinition End
+};
diff --git a/Engine/Plugins/Runtime/Soundscape/Source/SoundScapeEditor/Private/AssetTypeActions_SoundScapePalette.cpp b/Engine/Plugins/Runtime/Soundscape/Source/SoundScapeEditor/Private/AssetTypeActions_SoundScapePalette.cpp
deleted file mode 100644
diff --git a/Engine/Plugins/Runtime/Soundscape/Source/SoundScapeEditor/Private/AssetTypeActions_SoundScapePalette.h b/Engine/Plugins/Runtime/Soundscape/Source/SoundScapeEditor/Private/AssetTypeActions_SoundScapePalette.h
deleted file mode 100644
diff --git a/Engine/Plugins/Runtime/Soundscape/Source/SoundScapeEditor/Private/AssetTypeActions_SoundscapeColor.cpp b/Engine/Plugins/Runtime/Soundscape/Source/SoundScapeEditor/Private/AssetTypeActions_SoundscapeColor.cpp
deleted file mode 100644
diff --git a/Engine/Plugins/Runtime/Soundscape/Source/SoundScapeEditor/Private/AssetTypeActions_SoundscapeColor.h b/Engine/Plugins/Runtime/Soundscape/Source/SoundScapeEditor/Private/AssetTypeActions_SoundscapeColor.h
deleted file mode 100644
diff --git a/Engine/Plugins/Runtime/Soundscape/Source/SoundScapeEditor/Private/SoundScapeEditorModule.cpp b/Engine/Plugins/Runtime/Soundscape/Source/SoundScapeEditor/Private/SoundScapeEditorModule.cpp
--- a/Engine/Plugins/Runtime/Soundscape/Source/SoundScapeEditor/Private/SoundScapeEditorModule.cpp
+++ b/Engine/Plugins/Runtime/Soundscape/Source/SoundScapeEditor/Private/SoundScapeEditorModule.cpp
@@ -1,24 +1,13 @@
 // Copyright Epic Games, Inc. All Rights Reserved.
 
 #include "SoundScapeEditorModule.h"
-#include "AssetTypeActions_SoundscapeColor.h"
-#include "AssetTypeActions_SoundScapePalette.h"
-
-
 
 IMPLEMENT_MODULE(FSoundscapeEditorModule, SoundscapeEditor)
 
 void FSoundscapeEditorModule::StartupModule()
 {
-
-	// Register the audio editor asset type actions.
-	IAssetTools& AssetTools = FModuleManager::LoadModuleChecked<FAssetToolsModule>("AssetTools").Get();
-
-	AssetTools.RegisterAssetTypeActions(MakeShareable(new FAssetTypeActions_SoundscapeColor));
-	AssetTools.RegisterAssetTypeActions(MakeShareable(new FAssetTypeActions_SoundscapePalette));
 }
 
 void FSoundscapeEditorModule::ShutdownModule()
 {
 }
-
diff --git a/Engine/Plugins/Runtime/Soundscape/Source/SoundScapeEditor/SoundScapeEditor.Build.cs b/Engine/Plugins/Runtime/Soundscape/Source/SoundScapeEditor/SoundScapeEditor.Build.cs
--- a/Engine/Plugins/Runtime/Soundscape/Source/SoundScapeEditor/SoundScapeEditor.Build.cs
+++ b/Engine/Plugins/Runtime/Soundscape/Source/SoundScapeEditor/SoundScapeEditor.Build.cs
@@ -10,7 +10,6 @@ namespace UnrealBuildTool.Rules
 				new string[]
 				{
 				"Core",
-				"AssetTools",
 				"Soundscape",
 				"AudioMixer",
 
@@ -22,6 +21,7 @@ namespace UnrealBuildTool.Rules
 			PrivateDependencyModuleNames.AddRange(
 				new string[]
 				{
+				"AssetDefinition",
 				"CoreUObject",
 				"Engine",
 				"Slate",
diff --git a/Engine/Plugins/Runtime/Spatialization/Source/SpatializationEditor/Private/AssetDefinition_ITDSpatializationSettings.h b/Engine/Plugins/Runtime/Spatialization/Source/SpatializationEditor/Private/AssetDefinition_ITDSpatializationSettings.h
new file mode 100644
--- /dev/null
+++ b/Engine/Plugins/Runtime/Spatialization/Source/SpatializationEditor/Private/AssetDefinition_ITDSpatializationSettings.h
@@ -0,0 +1,29 @@
+// Copyright Epic Games, Inc. All Rights Reserved.
+
+#pragma once
+#include "AssetDefinitionDefault.h"
+#include "ITDSpatializationSourceSettings.h"
+#include "AssetDefinition_ITDSpatializationSettings.generated.h"
+
+UCLASS()
+class UAssetDefinition_ITDSpatializationSettings : public UAssetDefinitionDefault
+{
+	GENERATED_BODY()
+public:
+	// UAssetDefinition Begin
+	virtual FText GetAssetDisplayName() const override { return NSLOCTEXT("AssetDefinition", "AssetDefinition_ITDSpatializationSettings", "ITD Source Spatialization Settings"); }
+	virtual FLinearColor GetAssetColor() const override { return FColor(145, 145, 145); }
+	virtual TSoftClassPtr<UObject> GetAssetClass() const override { return UITDSpatializationSourceSettings::StaticClass(); }
+	virtual TConstArrayView<FAssetCategoryPath> GetAssetCategories() const override
+	{
+		static const auto Categories = 
+		{
+			FAssetCategoryPath(EAssetCategoryPaths::Audio,
+				NSLOCTEXT("AssetDefinition", "AssetDefinition_ITDSpatializationSettingsSubMenu", "Advanced"),
+				FCategoryPath(NSLOCTEXT("AssetDefinition", "AssetDefinition_ITDSpatializationSettingsSubMenuSection", "Binaural Spatialization"), ECategoryMenuType::Section))
+		};
+
+		return Categories;
+	}
+	// UAssetDefinition End
+};
diff --git a/Engine/Plugins/Runtime/Spatialization/Source/SpatializationEditor/Private/ITDSpatializationSourceSettingsFactory.cpp b/Engine/Plugins/Runtime/Spatialization/Source/SpatializationEditor/Private/ITDSpatializationSourceSettingsFactory.cpp
--- a/Engine/Plugins/Runtime/Spatialization/Source/SpatializationEditor/Private/ITDSpatializationSourceSettingsFactory.cpp
+++ b/Engine/Plugins/Runtime/Spatialization/Source/SpatializationEditor/Private/ITDSpatializationSourceSettingsFactory.cpp
@@ -6,35 +6,6 @@
 
 #include UE_INLINE_GENERATED_CPP_BY_NAME(ITDSpatializationSourceSettingsFactory)
 
-FText FAssetTypeActions_ITDSpatializationSettings::GetName() const
-{
-	return NSLOCTEXT("Spatialization", "FAssetTypeActions_ITDSpatializationSettings", "ITD Source Spatialization Settings");
-}
-
-FColor FAssetTypeActions_ITDSpatializationSettings::GetTypeColor() const
-{
-	return FColor(145, 145, 145);
-}
-
-UClass* FAssetTypeActions_ITDSpatializationSettings::GetSupportedClass() const
-{
-	return UITDSpatializationSourceSettings::StaticClass();
-}
-
-uint32 FAssetTypeActions_ITDSpatializationSettings::GetCategories()
-{
-	return EAssetTypeCategories::Sounds;
-}
-
-const TArray<FText>& FAssetTypeActions_ITDSpatializationSettings::GetSubMenus() const
-{
-	static const TArray<FText> ITDSubMenus
-	{
-		NSLOCTEXT("Spatialization", "AssetSpatializationSettingsSubMenu", "Spatialization")
-	};
-	return ITDSubMenus;
-}
-
 UITDSpatializationSettingsFactory::UITDSpatializationSettingsFactory(const FObjectInitializer& ObjectInitializer)
 	: Super(ObjectInitializer)
 {
diff --git a/Engine/Plugins/Runtime/Spatialization/Source/SpatializationEditor/Private/SpatializationEditorModule.cpp b/Engine/Plugins/Runtime/Spatialization/Source/SpatializationEditor/Private/SpatializationEditorModule.cpp
--- a/Engine/Plugins/Runtime/Spatialization/Source/SpatializationEditor/Private/SpatializationEditorModule.cpp
+++ b/Engine/Plugins/Runtime/Spatialization/Source/SpatializationEditor/Private/SpatializationEditorModule.cpp
@@ -1,33 +1,13 @@
 // Copyright Epic Games, Inc. All Rights Reserved.
 
 #include "SpatializationEditorModule.h"
-#include "ITDSpatializationSourceSettingsFactory.h"
-
-
-namespace
-{
-	static const FName AssetToolsName = TEXT("AssetTools");
-
-	template <typename T>
-	void AddAssetAction(IAssetTools& AssetTools, TArray<TSharedPtr<FAssetTypeActions_Base>>& AssetActions)
-	{
-		TSharedPtr<T> AssetAction = MakeShared<T>();
-		TSharedPtr<FAssetTypeActions_Base> AssetActionBase = StaticCastSharedPtr<FAssetTypeActions_Base>(AssetAction);
-		AssetTools.RegisterAssetTypeActions(AssetAction.ToSharedRef());
-		AssetActions.Add(AssetActionBase);
-	}
-} // namespace <>
-
 
 void FSpatializationEditorModule::StartupModule()
 {
-	IAssetTools& AssetTools = FModuleManager::LoadModuleChecked<FAssetToolsModule>(AssetToolsName).Get();
-	AddAssetAction<FAssetTypeActions_ITDSpatializationSettings>(AssetTools, AssetActions);
 }
 
 void FSpatializationEditorModule::ShutdownModule()
 {
-	AssetActions.Reset();
 }
 
 IMPLEMENT_MODULE(FSpatializationEditorModule, SpatializationEditor)
diff --git a/Engine/Plugins/Runtime/Spatialization/Source/SpatializationEditor/Public/ITDSpatializationSourceSettingsFactory.h b/Engine/Plugins/Runtime/Spatialization/Source/SpatializationEditor/Public/ITDSpatializationSourceSettingsFactory.h
--- a/Engine/Plugins/Runtime/Spatialization/Source/SpatializationEditor/Public/ITDSpatializationSourceSettingsFactory.h
+++ b/Engine/Plugins/Runtime/Spatialization/Source/SpatializationEditor/Public/ITDSpatializationSourceSettingsFactory.h
@@ -5,21 +5,8 @@
 #include "IAudioExtensionPlugin.h"
 #include "AssetToolsModule.h"
 #include "Factories/Factory.h"
-#include "AssetTypeActions_Base.h"
 #include "ITDSpatializationSourceSettingsFactory.generated.h"
 
-
-class FAssetTypeActions_ITDSpatializationSettings : public FAssetTypeActions_Base
-{
-public:
-	virtual FText GetName() const override;
-	virtual FColor GetTypeColor() const override;
-	virtual UClass* GetSupportedClass() const override;
-	virtual uint32 GetCategories() override;
-
-	virtual const TArray<FText>& GetSubMenus() const override;
-};
-
 UCLASS(MinimalAPI, hidecategories = Object)
 class UITDSpatializationSettingsFactory : public UFactory
 {
diff --git a/Engine/Plugins/Runtime/Spatialization/Source/SpatializationEditor/Public/SpatializationEditorModule.h b/Engine/Plugins/Runtime/Spatialization/Source/SpatializationEditor/Public/SpatializationEditorModule.h
--- a/Engine/Plugins/Runtime/Spatialization/Source/SpatializationEditor/Public/SpatializationEditorModule.h
+++ b/Engine/Plugins/Runtime/Spatialization/Source/SpatializationEditor/Public/SpatializationEditorModule.h
@@ -2,10 +2,7 @@
 #pragma once
 
 #include "CoreMinimal.h"
-
-#include "AssetTypeActions_Base.h"
 #include "Modules/ModuleManager.h"
-#include "ITDSpatializer.h"
 
 class FSpatializationEditorModule : public IModuleInterface
 {
@@ -13,7 +10,4 @@ public:
 	virtual void StartupModule() override;
 
 	virtual void ShutdownModule() override;
-
-private:
-	TArray<TSharedPtr<FAssetTypeActions_Base>> AssetActions;
 };
diff --git a/Engine/Plugins/Runtime/Spatialization/Source/SpatializationEditor/SpatializationEditor.Build.cs b/Engine/Plugins/Runtime/Spatialization/Source/SpatializationEditor/SpatializationEditor.Build.cs
--- a/Engine/Plugins/Runtime/Spatialization/Source/SpatializationEditor/SpatializationEditor.Build.cs
+++ b/Engine/Plugins/Runtime/Spatialization/Source/SpatializationEditor/SpatializationEditor.Build.cs
@@ -21,6 +21,7 @@ public class SpatializationEditor : ModuleRules
 		PrivateDependencyModuleNames.AddRange(
 			new string[]
 			{
+				"AssetDefinition",
 				"CoreUObject",
 				"EditorFramework",
 				"Engine",
diff --git a/Engine/Plugins/Runtime/Synthesis/Source/SynthesisEditor/Private/EpicSynth1PresetBank.cpp b/Engine/Plugins/Runtime/Synthesis/Source/SynthesisEditor/Private/EpicSynth1PresetBank.cpp
--- a/Engine/Plugins/Runtime/Synthesis/Source/SynthesisEditor/Private/EpicSynth1PresetBank.cpp
+++ b/Engine/Plugins/Runtime/Synthesis/Source/SynthesisEditor/Private/EpicSynth1PresetBank.cpp
@@ -18,7 +18,7 @@ const TArray<FText>& FAssetTypeActions_ModularSynthPresetBank::GetSubMenus() con
 {
 	static const TArray<FText> SubMenus
 	{
-		FText(LOCTEXT("AssetSoundSynthesisSubMenu", "Synthesis"))
+		FText(LOCTEXT("AssetSoundSynthesisSubMenu", "Legacy"))
 	};
 
 	return SubMenus;
diff --git a/Engine/Plugins/Runtime/Synthesis/Source/SynthesisEditor/Private/MonoWaveTablePresetBank.cpp b/Engine/Plugins/Runtime/Synthesis/Source/SynthesisEditor/Private/MonoWaveTablePresetBank.cpp
--- a/Engine/Plugins/Runtime/Synthesis/Source/SynthesisEditor/Private/MonoWaveTablePresetBank.cpp
+++ b/Engine/Plugins/Runtime/Synthesis/Source/SynthesisEditor/Private/MonoWaveTablePresetBank.cpp
@@ -18,7 +18,7 @@ const TArray<FText>& FAssetTypeActions_MonoWaveTableSynthPreset::GetSubMenus() c
 {
 	static const TArray<FText> SubMenus
 	{
-		FText(LOCTEXT("AssetSoundSynthesisSubMenu", "Synthesis"))
+		FText(LOCTEXT("AssetSoundSynthesisSubMenu", "Legacy"))
 	};
 
 	return SubMenus;
diff --git a/Engine/Plugins/Runtime/WaveTable/Source/WaveTableEditor/Private/AssetDefinition_WaveTableBank.cpp b/Engine/Plugins/Runtime/WaveTable/Source/WaveTableEditor/Private/AssetDefinition_WaveTableBank.cpp
new file mode 100644
--- /dev/null
+++ b/Engine/Plugins/Runtime/WaveTable/Source/WaveTableEditor/Private/AssetDefinition_WaveTableBank.cpp
@@ -0,0 +1,18 @@
+// Copyright Epic Games, Inc. All Rights Reserved.
+
+#include "AssetDefinition_WaveTableBank.h"
+#include "WaveTableBankEditor.h"
+
+EAssetCommandResult UAssetDefinition_WaveTableBank::OpenAssets(const FAssetOpenArgs& OpenArgs) const
+{
+	using namespace WaveTable::Editor;
+	const EToolkitMode::Type Mode = OpenArgs.ToolkitHost.IsValid() ? EToolkitMode::WorldCentric : EToolkitMode::Standalone;
+
+	for (UWaveTableBank* Patch : OpenArgs.LoadObjects<UWaveTableBank>())
+	{
+		TSharedRef<FBankEditor> Editor = MakeShared<FBankEditor>();
+		Editor->Init(Mode, OpenArgs.ToolkitHost, Patch);
+	}
+
+	return EAssetCommandResult::Handled;
+}
diff --git a/Engine/Plugins/Runtime/WaveTable/Source/WaveTableEditor/Private/AssetDefinition_WaveTableBank.h b/Engine/Plugins/Runtime/WaveTable/Source/WaveTableEditor/Private/AssetDefinition_WaveTableBank.h
new file mode 100644
--- /dev/null
+++ b/Engine/Plugins/Runtime/WaveTable/Source/WaveTableEditor/Private/AssetDefinition_WaveTableBank.h
@@ -0,0 +1,30 @@
+// Copyright Epic Games, Inc. All Rights Reserved.
+
+#pragma once
+#include "AssetDefinitionDefault.h"
+#include "WaveTableBank.h"
+#include "AssetDefinition_WaveTableBank.generated.h"
+
+UCLASS()
+class UAssetDefinition_WaveTableBank : public UAssetDefinitionDefault
+{
+	GENERATED_BODY()
+
+public:
+	// UAssetDefinition Begin
+	virtual FText GetAssetDisplayName() const override { return NSLOCTEXT("AssetDefinition", "AssetDefinition_WaveTableBank", "WaveTable Bank"); }
+	virtual FLinearColor GetAssetColor() const override { return FColor(232, 122, 0, 255); }
+	virtual TSoftClassPtr<UObject> GetAssetClass() const override { return UWaveTableBank::StaticClass(); }
+	virtual TConstArrayView<FAssetCategoryPath> GetAssetCategories() const override
+	{
+		static const auto Categories = 
+		{
+			FAssetCategoryPath(EAssetCategoryPaths::Audio,
+				NSLOCTEXT("AssetDefinition", "AssetDefinition_WaveTableBankSubMenu", "Advanced"),
+				FCategoryPath(NSLOCTEXT("AssetDefinition", "AssetDefinition_WaveTableBankSubMenuSection", "DSP Effects and Synthesis"), ECategoryMenuType::Section))
+		};
+		return Categories;
+	}
+	virtual EAssetCommandResult OpenAssets(const FAssetOpenArgs& OpenArgs) const override;
+	// UAssetDefinition End
+};
diff --git a/Engine/Plugins/Runtime/WaveTable/Source/WaveTableEditor/Private/AssetTypeActions_WaveTableBank.cpp b/Engine/Plugins/Runtime/WaveTable/Source/WaveTableEditor/Private/AssetTypeActions_WaveTableBank.cpp
deleted file mode 100644
diff --git a/Engine/Plugins/Runtime/WaveTable/Source/WaveTableEditor/Private/AssetTypeActions_WaveTableBank.h b/Engine/Plugins/Runtime/WaveTable/Source/WaveTableEditor/Private/AssetTypeActions_WaveTableBank.h
deleted file mode 100644
diff --git a/Engine/Plugins/Runtime/WaveTable/Source/WaveTableEditor/Private/WaveTableEditorModule.cpp b/Engine/Plugins/Runtime/WaveTable/Source/WaveTableEditor/Private/WaveTableEditorModule.cpp
--- a/Engine/Plugins/Runtime/WaveTable/Source/WaveTableEditor/Private/WaveTableEditorModule.cpp
+++ b/Engine/Plugins/Runtime/WaveTable/Source/WaveTableEditor/Private/WaveTableEditorModule.cpp
@@ -1,7 +1,6 @@
 // Copyright Epic Games, Inc. All Rights Reserved.
 #include "WaveTableEditorModule.h"
 
-#include "AssetTypeActions_WaveTableBank.h"
 #include "CurveDrawInfo.h"
 #include "ICurveEditorModule.h"
 #include "PropertyEditorModule.h"
@@ -89,9 +88,6 @@ namespace WaveTable::Editor
 			}
 		));
 
-		IAssetTools& AssetTools = FModuleManager::LoadModuleChecked<FAssetToolsModule>(AssetToolName).Get();
-		AddAssetAction<FAssetTypeActions_WaveTableBank>(AssetTools, AssetActions);
-
 		StyleSet = MakeShared<FSlateStyle>();
 	}
 
@@ -105,15 +101,6 @@ namespace WaveTable::Editor
 			PropertyModule.UnregisterCustomPropertyTypeLayout("WaveTableTransform");
 		}
 
-// 		if (FModuleManager::Get().IsModuleLoaded(AssetToolName))
-// 		{
-// 			IAssetTools& AssetTools = FModuleManager::GetModuleChecked<FAssetToolsModule>(AssetToolName).Get();
-// 			for (const TSharedPtr<FAssetTypeActions_Base>& ActionPtr : AssetActions)
-// 			{
-// 				AssetTools.UnregisterAssetTypeActions(ActionPtr);
-// 			}
-// 		}
-
 		if (FModuleManager::Get().IsModuleLoaded(CurveEditorName))
 		{
 			ICurveEditorModule& CurveEditorModule = FModuleManager::LoadModuleChecked<ICurveEditorModule>(CurveEditorName);
diff --git a/Engine/Plugins/Runtime/WaveTable/Source/WaveTableEditor/WaveTableEditor.Build.cs b/Engine/Plugins/Runtime/WaveTable/Source/WaveTableEditor/WaveTableEditor.Build.cs
--- a/Engine/Plugins/Runtime/WaveTable/Source/WaveTableEditor/WaveTableEditor.Build.cs
+++ b/Engine/Plugins/Runtime/WaveTable/Source/WaveTableEditor/WaveTableEditor.Build.cs
@@ -14,6 +14,7 @@ public class WaveTableEditor : ModuleRules
 
 		PrivateDependencyModuleNames.AddRange(new string[]
 		{
+			"AssetDefinition",
 			"AudioEditor",
 			"AudioExtensions",
 			"Core",
diff --git a/Engine/Source/Editor/AudioEditor/Private/AssetTypeActions/AssetDefinition_AudioBus.h b/Engine/Source/Editor/AudioEditor/Private/AssetTypeActions/AssetDefinition_AudioBus.h
--- a/Engine/Source/Editor/AudioEditor/Private/AssetTypeActions/AssetDefinition_AudioBus.h
+++ b/Engine/Source/Editor/AudioEditor/Private/AssetTypeActions/AssetDefinition_AudioBus.h
@@ -20,7 +20,12 @@ public:
 	virtual TSoftClassPtr<UObject> GetAssetClass() const override { return UAudioBus::StaticClass(); }
 	virtual TConstArrayView<FAssetCategoryPath> GetAssetCategories() const override
 	{
-		static const auto Categories = { FAssetCategoryPath(EAssetCategoryPaths::Audio, NSLOCTEXT("AssetTypeActions", "AssetAudioBusMenu", "Mix")) };
+		static const auto Categories = 
+			{
+				FAssetCategoryPath(EAssetCategoryPaths::Audio, 
+					NSLOCTEXT("AssetTypeActions", "AssetAudioBusMenu", "Advanced"), 
+					FCategoryPath(NSLOCTEXT("AssetTypeActions", "AssetAudioBusSubMenuSection", "Routing"), ECategoryMenuType::Section))
+			};
 		return Categories;
 	}
 	// UAssetDefinition End
diff --git a/Engine/Source/Editor/AudioEditor/Private/AssetTypeActions/AssetDefinition_DialogueVoice.h b/Engine/Source/Editor/AudioEditor/Private/AssetTypeActions/AssetDefinition_DialogueVoice.h
--- a/Engine/Source/Editor/AudioEditor/Private/AssetTypeActions/AssetDefinition_DialogueVoice.h
+++ b/Engine/Source/Editor/AudioEditor/Private/AssetTypeActions/AssetDefinition_DialogueVoice.h
@@ -20,7 +20,13 @@ public:
 	virtual TSoftClassPtr<UObject> GetAssetClass() const override { return UDialogueVoice::StaticClass(); }
 	virtual TConstArrayView<FAssetCategoryPath> GetAssetCategories() const override
 	{
-		static const auto Categories = { FAssetCategoryPath(EAssetCategoryPaths::Audio, NSLOCTEXT("AssetTypeActions", "AssetDialogueSubMenu", "Dialogue")) };
+		static const auto Categories = 
+			{
+				FAssetCategoryPath(EAssetCategoryPaths::Audio,
+					NSLOCTEXT("AssetDefinition", "AssetDefinition_DialogueVoiceSubMenu", "Advanced"),
+					FCategoryPath(NSLOCTEXT("AssetDefinition", "AssetDefinition_DialogueVoiceSubMenuSection", "Dialogue"), ECategoryMenuType::Section))
+			};
+
 		return Categories;
 	}
 	// UAssetDefinition End
diff --git a/Engine/Source/Editor/AudioEditor/Private/AssetTypeActions/AssetDefinition_DialogueWave.cpp b/Engine/Source/Editor/AudioEditor/Private/AssetTypeActions/AssetDefinition_DialogueWave.cpp
new file mode 100644
--- /dev/null
+++ b/Engine/Source/Editor/AudioEditor/Private/AssetTypeActions/AssetDefinition_DialogueWave.cpp
@@ -0,0 +1,285 @@
+// Copyright Epic Games, Inc. All Rights Reserved.
+
+#include "AssetDefinition_DialogueWave.h"
+#include "AnimationEditorUtils.h"
+#include "AssetToolsModule.h"
+#include "ContentBrowserMenuContexts.h"
+#include "ContentBrowserModule.h"
+#include "Factories/SoundCueFactoryNew.h"
+#include "Sound/SoundCue.h"
+#include "Toolkits/SimpleAssetEditor.h"
+
+#define LOCTEXT_NAMESPACE "AssetDefinition_DialogueWave"
+
+EAssetCommandResult UAssetDefinition_DialogueWave::OpenAssets(const FAssetOpenArgs& OpenArgs) const
+{
+	EToolkitMode::Type Mode = OpenArgs.ToolkitHost.IsValid() ? EToolkitMode::WorldCentric : EToolkitMode::Standalone;
+
+	for (UDialogueWave* DialogueWave : OpenArgs.LoadObjects<UDialogueWave>())
+	{
+		FSimpleAssetEditor::CreateEditor(Mode, Mode == EToolkitMode::WorldCentric ? OpenArgs.ToolkitHost : TSharedPtr<IToolkitHost>(), DialogueWave);
+	}
+	return EAssetCommandResult::Handled;
+}
+
+namespace MenuExtension::DialogueWave
+{
+	static bool CanExecutePlayCommand(const FToolMenuContext& InContext)
+	{
+		const UContentBrowserAssetContextMenuContext* CBContext = UContentBrowserAssetContextMenuContext::FindContextWithAssets(InContext);
+		if (CBContext == nullptr)
+		{
+			return false;
+		}
+
+		TArray<FAssetData> AssetData = CBContext->SelectedAssets;
+		if (AssetData.Num() != 1)
+		{
+			return false;
+		}
+
+		USoundBase* Sound = nullptr;
+		UDialogueWave* DialogueWave = CBContext->LoadFirstSelectedObject<UDialogueWave>();
+		
+		if (DialogueWave == nullptr)
+		{
+			return false;
+		}
+
+		for (int32 i = 0; i < DialogueWave->ContextMappings.Num(); ++i)
+		{
+			const FDialogueContextMapping& ContextMapping = DialogueWave->ContextMappings[i];
+
+			Sound = DialogueWave->GetWaveFromContext(ContextMapping.Context);
+			if (Sound != nullptr)
+			{
+				break;
+			}
+		}
+
+		return Sound != nullptr;
+	}
+
+	void StopSound()
+	{
+		GEditor->ResetPreviewAudioComponent();
+	}
+
+	static void PlaySound(UDialogueWave* DialogueWave)
+	{
+		USoundBase* Sound = nullptr;
+
+		for (int32 i = 0; i < DialogueWave->ContextMappings.Num(); ++i)
+		{
+			const FDialogueContextMapping& ContextMapping = DialogueWave->ContextMappings[i];
+
+			Sound = DialogueWave->GetWaveFromContext(ContextMapping.Context);
+			if (Sound != nullptr)
+			{
+				break;
+			}
+		}
+
+		if (Sound)
+		{
+			GEditor->PlayPreviewSound(Sound);
+		}
+		else
+		{
+			StopSound();
+		}
+	}
+
+	static void ExecutePlaySound(const FToolMenuContext& InContext)
+	{
+		if (UDialogueWave* DialogueWave = UContentBrowserAssetContextMenuContext::LoadSingleSelectedAsset<UDialogueWave>(InContext))
+		{
+			// Only play the first valid sound
+			PlaySound(DialogueWave);
+		}
+	}
+
+	static void ExecuteStopSound(const FToolMenuContext& InContext)
+	{
+		StopSound();
+	}
+
+	static bool IsCreateSoundCueActionVisible(const FToolMenuContext& InContext, bool bInAcceptMultipleEntries)
+	{
+		if (const UContentBrowserAssetContextMenuContext* CBContext = UContentBrowserAssetContextMenuContext::FindContextWithAssets(InContext))
+		{
+			TArray<FAssetData> AssetData = CBContext->SelectedAssets;
+			if (AssetData.IsEmpty())
+			{
+				return false;
+			}
+			
+			// If multiple entries are accepted only return true if the entries are more than 1
+			return bInAcceptMultipleEntries ? AssetData.Num() > 1 : AssetData.Num() == 1;
+		}
+
+		return false;
+	}
+
+	static void ExecuteCreateSoundCue(const FToolMenuContext& InContext, bool bCreateCueForEachDialogueWave)
+	{
+		if (const UContentBrowserAssetContextMenuContext* CBContext = UContentBrowserAssetContextMenuContext::FindContextWithAssets(InContext))
+		{
+			TArray<UDialogueWave*> Objects = CBContext->LoadSelectedObjects<UDialogueWave>();
+
+			if (Objects.IsEmpty())
+			{
+				return;
+			}
+
+			const FString DefaultSuffix = TEXT("_Cue");
+
+			if (Objects.Num() == 1 || !bCreateCueForEachDialogueWave)
+			{
+				UDialogueWave* Object = Objects[0];
+
+				if (Object)
+				{
+					// Determine an appropriate name
+					FString Name;
+					FString PackagePath;
+					AnimationEditorUtils::CreateUniqueAssetName(Object->GetOutermost()->GetName(), DefaultSuffix, PackagePath, Name);
+
+					// Create the factory used to generate the asset
+					USoundCueFactoryNew* Factory = NewObject<USoundCueFactoryNew>();
+					Factory->InitialDialogueWaves = { TWeakObjectPtr<UDialogueWave>(Object) };
+
+					FContentBrowserModule& ContentBrowserModule = FModuleManager::LoadModuleChecked<FContentBrowserModule>("ContentBrowser");
+					ContentBrowserModule.Get().CreateNewAsset(Name, FPackageName::GetLongPackagePath(PackagePath), USoundCue::StaticClass(), Factory);
+				}
+			}
+			else
+			{
+				TArray<UObject*> ObjectsToSync;
+
+				for (auto ObjIt = Objects.CreateConstIterator(); ObjIt; ++ObjIt)
+				{
+					UDialogueWave* Object = (*ObjIt);
+					if (Object)
+					{
+						FString Name;
+						FString PackageName;
+						AnimationEditorUtils::CreateUniqueAssetName(Object->GetOutermost()->GetName(), DefaultSuffix, PackageName, Name);
+
+						// Create the factory used to generate the asset
+						USoundCueFactoryNew* Factory = NewObject<USoundCueFactoryNew>();
+						Factory->InitialDialogueWaves.Add(Object);
+
+						FAssetToolsModule& AssetToolsModule = FModuleManager::GetModuleChecked<FAssetToolsModule>("AssetTools");
+						UObject* NewAsset = AssetToolsModule.Get().CreateAsset(Name, FPackageName::GetLongPackagePath(PackageName), USoundCue::StaticClass(), Factory);
+
+						if (NewAsset)
+						{
+							ObjectsToSync.Add(NewAsset);
+						}
+					}
+				}
+
+				if (ObjectsToSync.Num() > 0)
+				{
+					FContentBrowserModule& ContentBrowserModule = FModuleManager::LoadModuleChecked<FContentBrowserModule>("ContentBrowser");
+					ContentBrowserModule.Get().SyncBrowserToAssets(ObjectsToSync);
+				}
+			}
+		}
+	}
+
+	static FDelayedAutoRegisterHelper DelayedAutoRegister(EDelayedRegisterRunPhase::EndOfEngineInit, []
+	{
+		UToolMenus::RegisterStartupCallback(FSimpleMulticastDelegate::FDelegate::CreateLambda([]()
+		{
+			FToolMenuOwnerScoped OwnerScoped(UE_MODULE_NAME);
+			
+			// Dialogue Wave Action Registration
+			{
+				UToolMenu* Menu = UE::ContentBrowser::ExtendToolMenu_AssetContextMenu(UDialogueWave::StaticClass());
+
+				FToolMenuSection& PlaybackSection = Menu->FindOrAddSection("GetAssetActions");
+				PlaybackSection.AddDynamicEntry(NAME_None, FNewToolMenuSectionDelegate::CreateLambda([](FToolMenuSection& InSection)
+					{
+						if (const UContentBrowserAssetContextMenuContext* Context = UContentBrowserAssetContextMenuContext::FindContextWithAssets(InSection))
+						{
+							if (Context->SelectedAssets.Num() > 0)
+							{
+								const FAssetData& AssetData = Context->SelectedAssets[0];
+								if (AssetData.AssetClassPath == UDialogueWave::StaticClass()->GetClassPathName())
+								{
+									// Play
+									{
+										const TAttribute<FText> Label = LOCTEXT("DialogWave_PlaySound", "Play");
+										const TAttribute<FText> ToolTip = LOCTEXT("DialogWave_PlaySoundTooltip", "Plays the selected dialog wave.");
+										const FSlateIcon Icon = FSlateIcon(FAppStyle::GetAppStyleSetName(), "MediaAsset.AssetActions.Play.Small");
+
+										FToolUIAction UIAction;
+										UIAction.ExecuteAction = FToolMenuExecuteAction::CreateStatic(&ExecutePlaySound);
+										UIAction.CanExecuteAction = FToolMenuCanExecuteAction::CreateStatic(&CanExecutePlayCommand);
+										InSection.AddMenuEntry("DialogWave_PlaySound", Label, ToolTip, Icon, UIAction);
+									}
+
+									// Stop
+									{
+										const TAttribute<FText> Label = LOCTEXT("DialogWave_Stop", "Stop");
+										const TAttribute<FText> ToolTip = LOCTEXT("DialogWave_StopTooltip", "Stops the selected dialog wave(s).");
+										const FSlateIcon Icon = FSlateIcon(FAppStyle::GetAppStyleSetName(), "MediaAsset.AssetActions.Stop.Small");
+
+										FToolUIAction UIAction;
+										UIAction.ExecuteAction = FToolMenuExecuteAction::CreateStatic(&ExecuteStopSound);
+										InSection.AddMenuEntry("DialogWave_StopSound", Label, ToolTip, Icon, UIAction);
+									}
+
+									// Create Cue
+									{
+										constexpr bool bAcceptMultipleEntries = false;
+										constexpr bool bCreateCueForEachDialogueWave = true;
+										const TAttribute<FText> Label = LOCTEXT("DialogueWave_CreateCue", "Create Cue");
+										const TAttribute<FText> ToolTip = LOCTEXT("DialogueWave_CreateCueTooltip", "Creates a sound cue using this dialogue wave.");
+										const FSlateIcon Icon = FSlateIcon(FAppStyle::GetAppStyleSetName(), "ClassIcon.SoundCue");
+
+										FToolUIAction UIAction;
+										UIAction.ExecuteAction = FToolMenuExecuteAction::CreateStatic(&ExecuteCreateSoundCue, bCreateCueForEachDialogueWave);
+										UIAction.IsActionVisibleDelegate = FToolMenuIsActionButtonVisible::CreateStatic(&IsCreateSoundCueActionVisible, bAcceptMultipleEntries);
+										InSection.AddMenuEntry("DialogueWave_CreateCue", Label, ToolTip, Icon, UIAction);
+									}
+
+									// Create Single Cue
+									{
+										constexpr bool bAcceptMultipleEntries = true;
+										constexpr bool bCreateCueForEachDialogueWave = false;
+										const TAttribute<FText> Label = LOCTEXT("DialogueWave_CreateSingleCue", "Create Single Cue");
+										const TAttribute<FText> ToolTip = LOCTEXT("DialogueWave_CreateSingleCueTooltip", "Creates a single sound cue using these dialogue waves.");
+										const FSlateIcon Icon = FSlateIcon(FAppStyle::GetAppStyleSetName(), "ClassIcon.SoundCue");
+
+										FToolUIAction UIAction;
+										UIAction.ExecuteAction = FToolMenuExecuteAction::CreateStatic(&ExecuteCreateSoundCue, bCreateCueForEachDialogueWave);
+										UIAction.IsActionVisibleDelegate = FToolMenuIsActionButtonVisible::CreateStatic(&IsCreateSoundCueActionVisible, bAcceptMultipleEntries);
+										InSection.AddMenuEntry("DialogueWave_CreateSingleCue", Label, ToolTip, Icon, UIAction);
+									}
+
+									// Create Multiple Cue
+									{
+										constexpr bool bAcceptMultipleEntries = true;
+										constexpr bool bCreateCueForEachDialogueWave = true;
+										const TAttribute<FText> Label = LOCTEXT("DialogueWave_CreateMultipleCue", "Create Multiple Cues");
+										const TAttribute<FText> ToolTip = LOCTEXT("DialogueWave_CreateMultipleCueTooltip", "Creates multiple sound cues, one from each dialogue wave.");
+										const FSlateIcon Icon = FSlateIcon(FAppStyle::GetAppStyleSetName(), "ClassIcon.SoundCue");
+
+										FToolUIAction UIAction;
+										UIAction.ExecuteAction = FToolMenuExecuteAction::CreateStatic(&ExecuteCreateSoundCue, bCreateCueForEachDialogueWave);
+										UIAction.IsActionVisibleDelegate = FToolMenuIsActionButtonVisible::CreateStatic(&IsCreateSoundCueActionVisible, bAcceptMultipleEntries);
+										InSection.AddMenuEntry("DialogueWave_CreateMultipleCue", Label, ToolTip, Icon, UIAction);
+									}
+								}
+							}
+						}
+					}));
+			}
+		}));
+	});
+}
+
+#undef LOCTEXT_NAMESPACE
diff --git a/Engine/Source/Editor/AudioEditor/Private/AssetTypeActions/AssetDefinition_DialogueWave.h b/Engine/Source/Editor/AudioEditor/Private/AssetTypeActions/AssetDefinition_DialogueWave.h
new file mode 100644
--- /dev/null
+++ b/Engine/Source/Editor/AudioEditor/Private/AssetTypeActions/AssetDefinition_DialogueWave.h
@@ -0,0 +1,30 @@
+// Copyright Epic Games, Inc. All Rights Reserved.
+
+#pragma once
+#include "AssetTypeActions/AssetDefinitionDefault_AudioDiffable.h"
+#include "Sound/DialogueWave.h"
+#include "AssetDefinition_DialogueWave.generated.h"
+
+UCLASS()
+class UAssetDefinition_DialogueWave : public UAssetDefinitionDefault_AudioDiffable
+{
+	GENERATED_BODY()
+public:
+	// UAssetDefinition Begin
+	virtual FText GetAssetDisplayName() const override { return NSLOCTEXT("AssetDefinition", "AssetDefinition_DialogueWave", "Dialogue Wave"); }
+	virtual FLinearColor GetAssetColor() const override { return FColor(97, 85, 212); }
+	virtual TSoftClassPtr<UObject> GetAssetClass() const override { return UDialogueWave::StaticClass(); }
+	virtual TConstArrayView<FAssetCategoryPath> GetAssetCategories() const override
+	{
+		static const auto Categories = 
+		{
+			FAssetCategoryPath(EAssetCategoryPaths::Audio,
+				NSLOCTEXT("AssetDefinition", "AssetDefinition_DialogueWaveSubMenu", "Advanced"),
+				FCategoryPath(NSLOCTEXT("AssetDefinition", "AssetDefinition_DialogueWaveSubMenuSection", "Dialogue"), ECategoryMenuType::Section))
+		};
+
+		return Categories;
+	}
+	virtual EAssetCommandResult OpenAssets(const FAssetOpenArgs& OpenArgs) const override;
+	// UAssetDefinition End
+};
diff --git a/Engine/Source/Editor/AudioEditor/Private/AssetTypeActions/AssetDefinition_ReverbEffect.h b/Engine/Source/Editor/AudioEditor/Private/AssetTypeActions/AssetDefinition_ReverbEffect.h
--- a/Engine/Source/Editor/AudioEditor/Private/AssetTypeActions/AssetDefinition_ReverbEffect.h
+++ b/Engine/Source/Editor/AudioEditor/Private/AssetTypeActions/AssetDefinition_ReverbEffect.h
@@ -20,7 +20,7 @@ public:
 	virtual TSoftClassPtr<UObject> GetAssetClass() const override { return UReverbEffect::StaticClass(); }
 	virtual TConstArrayView<FAssetCategoryPath> GetAssetCategories() const override
 	{
-		static const auto Categories = { FAssetCategoryPath(EAssetCategoryPaths::Audio, NSLOCTEXT("AssetTypeActions", "AssetEffectSubMenu", "Effects")) };
+		static const auto Categories = { FAssetCategoryPath(EAssetCategoryPaths::Audio, NSLOCTEXT("AssetTypeActions", "AssetEffectSubMenu", "Legacy")) };
 		return Categories;
 	}
 	// UAssetDefinition End
diff --git a/Engine/Source/Editor/AudioEditor/Private/AssetTypeActions/AssetDefinition_SoundAttenuation.h b/Engine/Source/Editor/AudioEditor/Private/AssetTypeActions/AssetDefinition_SoundAttenuation.h
--- a/Engine/Source/Editor/AudioEditor/Private/AssetTypeActions/AssetDefinition_SoundAttenuation.h
+++ b/Engine/Source/Editor/AudioEditor/Private/AssetTypeActions/AssetDefinition_SoundAttenuation.h
@@ -16,13 +16,13 @@ class UAssetDefinition_SoundAttenuation : public UAssetDefinitionDefault_AudioDi
 
 public:
 	// UAssetDefinition Begin
-	virtual FText GetAssetDisplayName() const override { return NSLOCTEXT("AssetTypeActions", "AssetTypeActions_SoundAttenuation", "Sound Attenuation"); }
+	virtual FText GetAssetDisplayName() const override { return NSLOCTEXT("AssetTypeActions", "AssetTypeActions_SoundAttenuation", "Attenuation"); }
 	virtual FLinearColor GetAssetColor() const override { return FLinearColor(FColor(77, 120, 239)); }
 	virtual TSoftClassPtr<UObject> GetAssetClass() const override { return USoundAttenuation::StaticClass(); }
 	virtual TConstArrayView<FAssetCategoryPath> GetAssetCategories() const override
 	{
 		static const auto PinnedCategories = { EAssetCategoryPaths::Audio };
-		static const auto Categories = { FAssetCategoryPath(EAssetCategoryPaths::Audio, NSLOCTEXT("AssetTypeActions", "AssetSoundAttenuationSubMenu", "Spatialization")) };
+		static const auto Categories = { FAssetCategoryPath(EAssetCategoryPaths::Audio, NSLOCTEXT("AssetTypeActions", "AssetSoundAttenuationSubMenu", "Property Settings")) };
 		
 		if (GetDefault<UAudioEditorSettings>()->bPinSoundAttenuationInAssetMenu)
 		{
diff --git a/Engine/Source/Editor/AudioEditor/Private/AssetTypeActions/AssetDefinition_SoundClass.cpp b/Engine/Source/Editor/AudioEditor/Private/AssetTypeActions/AssetDefinition_SoundClass.cpp
new file mode 100644
--- /dev/null
+++ b/Engine/Source/Editor/AudioEditor/Private/AssetTypeActions/AssetDefinition_SoundClass.cpp
@@ -0,0 +1,219 @@
+// Copyright Epic Games, Inc. All Rights Reserved.
+
+#include "AssetDefinition_SoundClass.h"
+
+#include "AudioDeviceManager.h"
+#include "AudioEditorModule.h"
+#include "ContentBrowserMenuContexts.h"
+#include "ToolMenus.h"
+#include "Audio/AudioDebug.h"
+
+#define LOCTEXT_NAMESPACE "AssetDefinition_SoundClass"
+
+EAssetCommandResult UAssetDefinition_SoundClass::OpenAssets(const FAssetOpenArgs& OpenArgs) const
+{
+	const EToolkitMode::Type Mode = OpenArgs.ToolkitHost.IsValid() ? EToolkitMode::WorldCentric : EToolkitMode::Standalone;
+
+	for (USoundClass* SoundClass : OpenArgs.LoadObjects<USoundClass>())
+	{
+		IAudioEditorModule* AudioEditorModule = &FModuleManager::LoadModuleChecked<IAudioEditorModule>( "AudioEditor" );
+		AudioEditorModule->CreateSoundClassEditor(Mode, OpenArgs.ToolkitHost, SoundClass);
+	}
+	return EAssetCommandResult::Handled;
+}
+
+namespace MenuExtension::SoundClass
+{
+	static void ExecuteMute(const FToolMenuContext& InContext)
+	{
+#if ENABLE_AUDIO_DEBUG
+		const UContentBrowserAssetContextMenuContext* CBContext = UContentBrowserAssetContextMenuContext::FindContextWithAssets(InContext);
+		if (CBContext == nullptr)
+		{
+			return;
+		}
+
+		if (FAudioDeviceManager* ADM = GEditor->GetAudioDeviceManager())
+		{
+			for (USoundClass* SoundClass : CBContext->LoadSelectedObjects<USoundClass>())
+			{
+				ADM->GetDebugger().ToggleMuteSoundClass(SoundClass->GetFName());
+			}
+		}
+#endif
+	}
+	
+	static void ExecuteSolo(const FToolMenuContext& InContext)
+	{
+#if ENABLE_AUDIO_DEBUG
+		const UContentBrowserAssetContextMenuContext* CBContext = UContentBrowserAssetContextMenuContext::FindContextWithAssets(InContext);
+		if (CBContext == nullptr)
+		{
+			return;
+		}
+
+		if (FAudioDeviceManager* ADM = GEditor->GetAudioDeviceManager())
+		{
+			for (USoundClass* SoundClass : CBContext->LoadSelectedObjects<USoundClass>())
+			{
+				ADM->GetDebugger().ToggleSoloSoundClass(SoundClass->GetFName());
+			}
+		}
+#endif
+	}
+	
+	static ECheckBoxState IsActionCheckedMute(const FToolMenuContext& InContext)
+	{
+#if ENABLE_AUDIO_DEBUG
+		const UContentBrowserAssetContextMenuContext* CBContext = UContentBrowserAssetContextMenuContext::FindContextWithAssets(InContext);
+		if (CBContext == nullptr)
+		{
+			return ECheckBoxState::Undetermined;
+		}
+
+		if (FAudioDeviceManager* ADM = GEditor->GetAudioDeviceManager())
+		{
+			for (USoundClass* SoundClass : CBContext->LoadSelectedObjects<USoundClass>())
+			{
+				if (ADM->GetDebugger().IsMuteSoundClass(SoundClass->GetFName()))
+				{
+					return ECheckBoxState::Checked;
+				}
+			}
+		}
+#endif
+		return ECheckBoxState::Unchecked;
+	}
+	
+	static ECheckBoxState IsActionCheckedSolo(const FToolMenuContext& InContext)
+	{
+#if ENABLE_AUDIO_DEBUG
+		const UContentBrowserAssetContextMenuContext* CBContext = UContentBrowserAssetContextMenuContext::FindContextWithAssets(InContext);
+		if (CBContext == nullptr)
+		{
+			return ECheckBoxState::Undetermined;
+		}
+
+		if (FAudioDeviceManager* ADM = GEditor->GetAudioDeviceManager())
+		{
+			for (USoundClass* SoundClass : CBContext->LoadSelectedObjects<USoundClass>())
+			{
+				if (ADM->GetDebugger().IsSoloSoundClass(SoundClass->GetFName()))
+				{
+					return ECheckBoxState::Checked;
+				}
+			}
+		}
+#endif
+		return ECheckBoxState::Unchecked;
+	}
+	
+	static bool CanExecuteMuteCommand(const FToolMenuContext& InContext)
+	{
+#if ENABLE_AUDIO_DEBUG
+		const UContentBrowserAssetContextMenuContext* CBContext = UContentBrowserAssetContextMenuContext::FindContextWithAssets(InContext);
+		if (CBContext == nullptr)
+		{
+			return false;
+		}
+
+		if (FAudioDeviceManager* ADM = GEditor->GetAudioDeviceManager())
+		{
+			// Allow muting if we're not Soloing.
+			for (USoundClass* SoundClass : CBContext->LoadSelectedObjects<USoundClass>())
+			{
+				if (ADM->GetDebugger().IsSoloSoundClass(SoundClass->GetFName()))
+				{
+					return false;
+				}
+			}
+			// Ok.
+			return true;
+		}
+#endif
+		return false;
+	}
+	
+	static bool CanExecuteSoloCommand(const FToolMenuContext& InContext)
+	{
+#if ENABLE_AUDIO_DEBUG
+		const UContentBrowserAssetContextMenuContext* CBContext = UContentBrowserAssetContextMenuContext::FindContextWithAssets(InContext);
+		if (CBContext == nullptr)
+		{
+			return false;
+		}
+
+		if (FAudioDeviceManager* ADM = GEditor->GetAudioDeviceManager())
+		{
+			// Allow Soloing if we're not Muting.
+			for (USoundClass* SoundClass : CBContext->LoadSelectedObjects<USoundClass>())
+			{
+				if (ADM->GetDebugger().IsMuteSoundClass(SoundClass->GetFName()))
+				{
+					return false;
+				}
+			}
+			// Ok.
+			return true;
+		}
+#endif
+		return false;
+	}
+	
+	static FDelayedAutoRegisterHelper DelayedAutoRegister(EDelayedRegisterRunPhase::EndOfEngineInit, []
+	{
+		UToolMenus::RegisterStartupCallback(FSimpleMulticastDelegate::FDelegate::CreateLambda([]()
+		{
+			FToolMenuOwnerScoped OwnerScoped(UE_MODULE_NAME);
+			
+			// Sound Class Action Registration
+			{
+				UToolMenu* Menu = UE::ContentBrowser::ExtendToolMenu_AssetContextMenu(USoundClass::StaticClass());
+
+				FToolMenuSection& PlaybackSection = Menu->FindOrAddSection("GetAssetActions");
+				PlaybackSection.AddDynamicEntry(NAME_None, FNewToolMenuSectionDelegate::CreateLambda([](FToolMenuSection& InSection)
+					{
+						if (const UContentBrowserAssetContextMenuContext* Context = UContentBrowserAssetContextMenuContext::FindContextWithAssets(InSection))
+						{
+							if (Context->SelectedAssets.Num() > 0)
+							{
+								const FAssetData& AssetData = Context->SelectedAssets[0];
+								if (AssetData.AssetClassPath == USoundClass::StaticClass()->GetClassPathName())
+								{
+									// Mute
+									{
+										const TAttribute<FText> Label = LOCTEXT("Sound_MuteSoundClass", "Mute");
+										const TAttribute<FText> ToolTip = LOCTEXT("Sound_MuteSoundClassTooltip", "Mutes anything using this SoundClass");
+										const FSlateIcon Icon = FSlateIcon(FAppStyle::GetAppStyleSetName(), "MediaAsset.AssetActions.Mute.Small");
+										constexpr EUserInterfaceActionType ActionType = EUserInterfaceActionType::ToggleButton;
+
+										FToolUIAction UIAction;
+										UIAction.ExecuteAction = FToolMenuExecuteAction::CreateStatic(&ExecuteMute);
+										UIAction.CanExecuteAction = FToolMenuCanExecuteAction::CreateStatic(&CanExecuteMuteCommand);
+										UIAction.GetActionCheckState = FToolMenuGetActionCheckState::CreateStatic(&IsActionCheckedMute);
+										InSection.AddMenuEntry("SoundClass_Mute", Label, ToolTip, Icon, UIAction, ActionType);
+									}
+
+									// Solo
+									{
+										const TAttribute<FText> Label = LOCTEXT("Sound_SoloSoundClass", "Solo");
+										const TAttribute<FText> ToolTip = LOCTEXT("Sound_SoloSoundClassTooltip", "Solos the selected sounds.");
+										const FSlateIcon Icon = FSlateIcon(FAppStyle::GetAppStyleSetName(), "MediaAsset.AssetActions.Solo.Small");
+										constexpr EUserInterfaceActionType ActionType = EUserInterfaceActionType::ToggleButton;
+
+										FToolUIAction UIAction;
+										UIAction.ExecuteAction = FToolMenuExecuteAction::CreateStatic(&ExecuteSolo);
+										UIAction.CanExecuteAction = FToolMenuIsActionButtonVisible::CreateStatic(&CanExecuteSoloCommand);
+										UIAction.GetActionCheckState = FToolMenuGetActionCheckState::CreateStatic(&IsActionCheckedSolo);
+										InSection.AddMenuEntry("SoundClass_Solo", Label, ToolTip, Icon, UIAction, ActionType);
+									}
+								}
+							}
+						}
+					}));
+			}
+		}));
+	});
+}
+
+#undef LOCTEXT_NAMESPACE
diff --git a/Engine/Source/Editor/AudioEditor/Private/AssetTypeActions/AssetDefinition_SoundClass.h b/Engine/Source/Editor/AudioEditor/Private/AssetTypeActions/AssetDefinition_SoundClass.h
new file mode 100644
--- /dev/null
+++ b/Engine/Source/Editor/AudioEditor/Private/AssetTypeActions/AssetDefinition_SoundClass.h
@@ -0,0 +1,30 @@
+// Copyright Epic Games, Inc. All Rights Reserved.
+
+#pragma once
+#include "AssetTypeActions/AssetDefinitionDefault_AudioDiffable.h"
+#include "Sound/SoundClass.h"
+#include "AssetDefinition_SoundClass.generated.h"
+
+UCLASS()
+class UAssetDefinition_SoundClass : public UAssetDefinitionDefault_AudioDiffable
+{
+	GENERATED_BODY()
+public:
+	// UAssetDefinition Begin
+	virtual FText GetAssetDisplayName() const override { return NSLOCTEXT("AssetDefinition", "AssetDefinition_SoundClass", "Sound Class"); }
+	virtual FLinearColor GetAssetColor() const override { return FColor(255, 175, 0); }
+	virtual TSoftClassPtr<UObject> GetAssetClass() const override { return USoundClass::StaticClass(); }
+	virtual TConstArrayView<FAssetCategoryPath> GetAssetCategories() const override
+	{
+		static const auto Categories = 
+		{
+			FAssetCategoryPath(EAssetCategoryPaths::Audio,
+				NSLOCTEXT("AssetDefinition", "AssetDefinition_SoundClassSubMenu", "Advanced"),
+				FCategoryPath(NSLOCTEXT("AssetDefinition", "AssetDefinition_SoundClassSubMenuSection", "Properties"), ECategoryMenuType::Section))
+		};
+
+		return Categories;
+	}
+	virtual EAssetCommandResult OpenAssets(const FAssetOpenArgs& OpenArgs) const override;
+	// UAssetDefinition End
+};
diff --git a/Engine/Source/Editor/AudioEditor/Private/AssetTypeActions/AssetDefinition_SoundConcurrency.h b/Engine/Source/Editor/AudioEditor/Private/AssetTypeActions/AssetDefinition_SoundConcurrency.h
--- a/Engine/Source/Editor/AudioEditor/Private/AssetTypeActions/AssetDefinition_SoundConcurrency.h
+++ b/Engine/Source/Editor/AudioEditor/Private/AssetTypeActions/AssetDefinition_SoundConcurrency.h
@@ -16,13 +16,13 @@ class UAssetDefinition_SoundConcurrency : public UAssetDefinitionDefault_AudioDi
 
 public:
 	// UAssetDefinition Begin
-	virtual FText GetAssetDisplayName() const override { return NSLOCTEXT("AssetTypeActions", "AssetTypeActions_SoundConcurrency", "Sound Concurrency"); }
+	virtual FText GetAssetDisplayName() const override { return NSLOCTEXT("AssetTypeActions", "AssetTypeActions_SoundConcurrency", "Concurrency"); }
 	virtual FLinearColor GetAssetColor() const override { return FLinearColor(FColor(77, 100, 139)); }
 	virtual TSoftClassPtr<UObject> GetAssetClass() const override { return USoundConcurrency::StaticClass(); }
 	virtual TConstArrayView<FAssetCategoryPath> GetAssetCategories() const override
 	{
 		static const auto PinnedCategories = { EAssetCategoryPaths::Audio };
-		static const auto Categories = { FAssetCategoryPath(EAssetCategoryPaths::Audio, NSLOCTEXT("AssetTypeActions", "AssetSoundConcurrencySubMenu", "Mix")) };
+		static const auto Categories = { FAssetCategoryPath(EAssetCategoryPaths::Audio, NSLOCTEXT("AssetTypeActions", "AssetSoundConcurrencySubMenu", "Property Settings")) };
 
 		if (GetDefault<UAudioEditorSettings>()->bPinSoundConcurrencyInAssetMenu)
 		{
diff --git a/Engine/Source/Editor/AudioEditor/Private/AssetTypeActions/AssetDefinition_SoundCue.cpp b/Engine/Source/Editor/AudioEditor/Private/AssetTypeActions/AssetDefinition_SoundCue.cpp
--- a/Engine/Source/Editor/AudioEditor/Private/AssetTypeActions/AssetDefinition_SoundCue.cpp
+++ b/Engine/Source/Editor/AudioEditor/Private/AssetTypeActions/AssetDefinition_SoundCue.cpp
@@ -41,15 +41,8 @@ class UClass;
 
 TConstArrayView<FAssetCategoryPath> UAssetDefinition_SoundCue::GetAssetCategories() const
 {
-	/*
-	static const auto Categories = {
-		EAssetCategoryPaths::Audio,
-		EAssetCategoryPaths::Audio / LOCTEXT("AssetSoundCueSubMenu", "Source")
-	};
-	*/
-
 	static const auto Pinned_Categories = { EAssetCategoryPaths::Audio };
-	static const auto Categories = { EAssetCategoryPaths::Audio / LOCTEXT("AssetSoundCueSubMenu", "Source") };
+	static const auto Categories = { FAssetCategoryPath(EAssetCategoryPaths::Audio, LOCTEXT("SoundCue_SubMenu", "Sources")) };
 	
 	if (GetDefault<UAudioEditorSettings>()->bPinSoundCueInAssetMenu)
 	{
@@ -162,7 +155,7 @@ namespace MenuExtension_SoundCue
 				{
 					const TAttribute<FText> Label = LOCTEXT("SoundCue_ConsolidateAttenuation", "Consolidate Attenuation");
 					const TAttribute<FText> ToolTip = LOCTEXT("SoundCue_ConsolidateAttenuationTooltip", "Creates shared attenuation packages for sound cues with identical override attenuation settings.");
-					const FSlateIcon Icon = FSlateIcon(FAppStyle::GetAppStyleSetName(), "ClassIcon.SoundCue");
+					const FSlateIcon Icon = FSlateIcon(FAppStyle::GetAppStyleSetName(), "MediaAsset.AssetActions.SoundConsolidation.Small");
 
 					FToolUIAction UIAction;
 					UIAction.ExecuteAction = FToolMenuExecuteAction::CreateStatic(&ExecuteConsolidateAttenuation);
diff --git a/Engine/Source/Editor/AudioEditor/Private/AssetTypeActions/AssetDefinition_SoundEffectPreset.cpp b/Engine/Source/Editor/AudioEditor/Private/AssetTypeActions/AssetDefinition_SoundEffectPreset.cpp
new file mode 100644
--- /dev/null
+++ b/Engine/Source/Editor/AudioEditor/Private/AssetTypeActions/AssetDefinition_SoundEffectPreset.cpp
@@ -0,0 +1,132 @@
+// Copyright Epic Games, Inc. All Rights Reserved.
+
+#include "AssetDefinition_SoundEffectPreset.h"
+
+#include "Audio/AudioWidgetSubsystem.h"
+#include "Editor.h"
+#include "Editors/SoundEffectPresetEditor.h"
+#include "Templates/SubclassOf.h"
+#include "UObject/Class.h"
+#include "UObject/UObjectBaseUtility.h"
+
+#define LOCTEXT_NAMESPACE "UAssetDefinition_SoundEffectPresetDynamic"
+
+FText UAssetDefinition_SoundEffectPresetDynamic::GetAssetDisplayName() const
+{
+	if (!EffectPresetCDO.IsValid())
+	{
+		return LOCTEXT("AssetDefinition_SoundEffectPresetDynamic", "Sound Effect Preset Base (Invalid)");
+	}
+
+	const FText AssetActionName = EffectPresetCDO->GetAssetActionName();
+	if (AssetActionName.IsEmpty())
+	{
+		FString ClassName;
+		EffectPresetCDO->GetClass()->GetName(ClassName);
+		ensureMsgf(false, TEXT("U%sGetAssetActionName not implemented. Please check that EFFECT_PRESET_METHODS(EffectClassName) is at the top of the declaration of %s."), *ClassName, *ClassName);
+		const FString DefaultName = ClassName + FString(TEXT(" (Error: EFFECT_PRESET_METHODS() Not Used in Class Declaration)"));
+		return FText::FromString(DefaultName);
+	}
+
+	return EffectPresetCDO->GetAssetActionName();
+}
+
+FLinearColor UAssetDefinition_SoundEffectPresetDynamic::GetAssetColor() const
+{
+	if (!EffectPresetCDO.IsValid())
+	{
+		return FColor::White;
+	}
+
+	return EffectPresetCDO->GetPresetColor();
+}
+
+TSoftClassPtr<UObject> UAssetDefinition_SoundEffectPresetDynamic::GetAssetClass() const
+{
+	if (!EffectPresetCDO.IsValid())
+	{
+		return USoundEffectPreset::StaticClass();
+	}
+
+	const UClass* SupportedClass = EffectPresetCDO->GetSupportedClass();
+	if (SupportedClass == nullptr)
+	{
+		FString ClassName;
+		EffectPresetCDO->GetClass()->GetName(ClassName);
+		ensureMsgf(false, TEXT("U%s::GetSupportedClass not implemented. Please check that EFFECT_PRESET_METHODS(EffectClassName) is at the top of the declaration of %s."), *ClassName, *ClassName);
+		return EffectPresetCDO->GetClass();
+	}
+
+	return SupportedClass;
+}
+
+TConstArrayView<FAssetCategoryPath> UAssetDefinition_SoundEffectPresetDynamic::GetAssetCategories() const
+{
+	if (!EffectPresetCDO.IsValid() || !EffectPresetCDO->CanFilter())
+	{
+		static const TArray<FAssetCategoryPath> EmptyCategory = TArray<FAssetCategoryPath>();
+		return EmptyCategory;
+	}
+
+	static const TArray<FAssetCategoryPath> Categories = 
+	{
+		FAssetCategoryPath(EAssetCategoryPaths::Audio,
+			NSLOCTEXT("AssetDefinition", "AssetDefinition_SoundEffectPresetSubMenu", "Advanced"),
+			FCategoryPath(NSLOCTEXT("AssetDefinition", "AssetDefinition_SoundEffectPresetSubMenuSection", "DSP Effects and Synthesis"), ECategoryMenuType::Section))
+	};
+
+	return Categories;
+}
+
+EAssetCommandResult UAssetDefinition_SoundEffectPresetDynamic::OpenAssets(const FAssetOpenArgs& OpenArgs) const
+{
+	EToolkitMode::Type Mode = OpenArgs.ToolkitHost.IsValid() ? EToolkitMode::WorldCentric : EToolkitMode::Standalone;
+
+	TArray<UObject*> Objects = OpenArgs.LoadObjects<UObject>();
+	if (UWorld* World = GEditor ? GEditor->GetEditorWorldContext().World() : nullptr)
+	{
+		if (UAudioWidgetSubsystem* WidgetSubsystem = GEngine ? GEngine->GetEngineSubsystem<UAudioWidgetSubsystem>() : nullptr)
+		{
+			for (UObject* Object : Objects)
+			{
+				USoundEffectPreset* Preset = Cast<USoundEffectPreset>(Object);
+				if (!Preset)
+				{
+					continue;
+				}
+
+				auto FilterFunction = [InPresetClass = Object->GetClass()](UUserWidget* UserWidget)
+					{
+						TSubclassOf<USoundEffectPreset> PresetClass = ISoundEffectPresetWidgetInterface::Execute_GetClass(UserWidget);
+						while (PresetClass)
+						{
+							if (PresetClass == InPresetClass)
+							{
+								return true;
+							}
+
+							PresetClass = PresetClass->GetSuperClass();
+						}
+
+						return false;
+					};
+
+				TArray<UUserWidget*> UserWidgets = WidgetSubsystem->CreateUserWidgets(*World, USoundEffectPresetWidgetInterface::StaticClass(), FilterFunction);
+				if (!UserWidgets.IsEmpty())
+				{
+					TSharedRef<FSoundEffectPresetEditor> PresetEditor = MakeShared<FSoundEffectPresetEditor>();
+					PresetEditor->Init(Mode, OpenArgs.ToolkitHost, Preset, UserWidgets);
+				}
+			}
+		}
+	}
+	
+	return EAssetCommandResult::Handled;
+}
+
+void UAssetDefinition_SoundEffectPresetDynamic::Initialize(TSubclassOf<USoundEffectPreset> InClass)
+{
+	EffectPresetCDO = TStrongObjectPtr(InClass->GetDefaultObject<USoundEffectPreset>());
+}
+
+#undef LOCTEXT_NAMESPACE
diff --git a/Engine/Source/Editor/AudioEditor/Private/AssetTypeActions/AssetDefinition_SoundEffectPreset.h b/Engine/Source/Editor/AudioEditor/Private/AssetTypeActions/AssetDefinition_SoundEffectPreset.h
new file mode 100644
--- /dev/null
+++ b/Engine/Source/Editor/AudioEditor/Private/AssetTypeActions/AssetDefinition_SoundEffectPreset.h
@@ -0,0 +1,96 @@
+// Copyright Epic Games, Inc. All Rights Reserved.
+
+#pragma once
+#include "AssetTypeActions/AssetDefinitionDefault_AudioDiffable.h"
+#include "Sound/SoundEffectSource.h"
+#include "Sound/SoundEffectSubmix.h"
+#include "AssetDefinition_SoundEffectPreset.generated.h"
+
+UCLASS()
+class UAssetDefinition_SoundEffectSubmixPreset : public UAssetDefinitionDefault_AudioDiffable
+{
+	GENERATED_BODY()
+public:
+	// UAssetDefinition Begin
+	virtual FText GetAssetDisplayName() const override { return NSLOCTEXT("AssetDefinition", "AssetDefinition_SoundEffectSubmixPreset", "Submix Effect Preset"); }
+	virtual FLinearColor GetAssetColor() const override { return FColor(99, 63, 56); }
+	virtual TSoftClassPtr<UObject> GetAssetClass() const override { return USoundEffectSubmixPreset::StaticClass(); }
+	virtual TConstArrayView<FAssetCategoryPath> GetAssetCategories() const override
+	{
+		static const auto Categories = 
+		{
+			FAssetCategoryPath(EAssetCategoryPaths::Audio,
+				NSLOCTEXT("AssetDefinition", "AssetDefinition_SoundEffectSubmixPresetSubMenu", "Advanced"),
+				FCategoryPath(NSLOCTEXT("AssetDefinition", "AssetDefinition_SoundEffectSubmixPresetSubMenuSection", "DSP Effects and Synthesis"), ECategoryMenuType::Section))
+		};
+
+		return Categories;
+	}
+	// UAssetDefinition End
+};
+
+UCLASS()
+class UAssetDefinition_SoundEffectSourcePreset : public UAssetDefinitionDefault_AudioDiffable
+{
+	GENERATED_BODY()
+public:
+	// UAssetDefinition Begin
+	virtual FText GetAssetDisplayName() const override { return NSLOCTEXT("AssetDefinition", "AssetDefinition_SoundEffectSourcePreset", "Source Effect Preset"); }
+	virtual FLinearColor GetAssetColor() const override { return FColor(72, 185, 187); }
+	virtual TSoftClassPtr<UObject> GetAssetClass() const override { return USoundEffectSourcePreset::StaticClass(); }
+	virtual TConstArrayView<FAssetCategoryPath> GetAssetCategories() const override
+	{
+		static const auto Categories = 
+		{
+			FAssetCategoryPath(EAssetCategoryPaths::Audio,
+				NSLOCTEXT("AssetDefinition", "AssetDefinition_SoundEffectSourcePresetSubMenu", "Advanced"),
+				FCategoryPath(NSLOCTEXT("AssetDefinition", "AssetDefinition_SoundEffectSourcePresetSubMenuSection", "DSP Effects and Synthesis"), ECategoryMenuType::Section))
+		};
+
+		return Categories;
+	}
+	// UAssetDefinition End
+};
+
+UCLASS()
+class UAssetDefinition_SoundEffectSourcePresetChain : public UAssetDefinitionDefault_AudioDiffable
+{
+	GENERATED_BODY()
+public:
+	// UAssetDefinition Begin
+	virtual FText GetAssetDisplayName() const override { return NSLOCTEXT("AssetDefinition", "AssetDefinition_SoundEffectSourcePresetChain", "Source Effect Preset Chain"); }
+	virtual FLinearColor GetAssetColor() const override { return FColor(51, 107, 142); }
+	virtual TSoftClassPtr<UObject> GetAssetClass() const override { return USoundEffectSourcePresetChain::StaticClass(); }
+	virtual TConstArrayView<FAssetCategoryPath> GetAssetCategories() const override
+	{
+		static const auto Categories = 
+		{
+			FAssetCategoryPath(EAssetCategoryPaths::Audio,
+				NSLOCTEXT("AssetDefinition", "AssetDefinition_SoundEffectSourcePresetChainSubMenu", "Advanced"),
+				FCategoryPath(NSLOCTEXT("AssetDefinition", "AssetDefinition_SoundEffectSourcePresetChainSubMenuSection", "DSP Effects and Synthesis"), ECategoryMenuType::Section))
+		};
+
+		return Categories;
+	}
+	// UAssetDefinition End
+};
+
+UCLASS()
+class UAssetDefinition_SoundEffectPresetDynamic : public UAssetDefinitionDefault_AudioDiffable
+{
+	GENERATED_BODY()
+public:
+	// UAssetDefinition Begin
+	virtual FText GetAssetDisplayName() const override;
+	virtual FLinearColor GetAssetColor() const override;
+	virtual TSoftClassPtr<UObject> GetAssetClass() const override;
+	virtual TConstArrayView<FAssetCategoryPath> GetAssetCategories() const override;
+	virtual EAssetCommandResult OpenAssets(const FAssetOpenArgs& OpenArgs) const override;
+	virtual bool CanRegisterStatically() const override { return false; }
+	// UAssetDefinition End
+
+	void Initialize(TSubclassOf<USoundEffectPreset> InClass);
+
+private:
+	TStrongObjectPtr<USoundEffectPreset> EffectPresetCDO;
+};
diff --git a/Engine/Source/Editor/AudioEditor/Private/AssetTypeActions/AssetDefinition_SoundMix.h b/Engine/Source/Editor/AudioEditor/Private/AssetTypeActions/AssetDefinition_SoundMix.h
--- a/Engine/Source/Editor/AudioEditor/Private/AssetTypeActions/AssetDefinition_SoundMix.h
+++ b/Engine/Source/Editor/AudioEditor/Private/AssetTypeActions/AssetDefinition_SoundMix.h
@@ -20,7 +20,7 @@ public:
 	virtual TSoftClassPtr<UObject> GetAssetClass() const override { return USoundMix::StaticClass(); }
 	virtual TConstArrayView<FAssetCategoryPath> GetAssetCategories() const override
 	{
-		static const auto Categories = { FAssetCategoryPath(EAssetCategoryPaths::Audio, NSLOCTEXT("AssetTypeActions", "AssetSoundClassSubMenu", "Classes")) };
+		static const auto Categories = { FAssetCategoryPath(EAssetCategoryPaths::Audio, NSLOCTEXT("AssetDefinition", "SoundMix_SubMenu", "Legacy")) };
 		return Categories;
 	}
 	// UAssetDefinition End
diff --git a/Engine/Source/Editor/AudioEditor/Private/AssetTypeActions/AssetDefinition_SoundSourceBus.h b/Engine/Source/Editor/AudioEditor/Private/AssetTypeActions/AssetDefinition_SoundSourceBus.h
--- a/Engine/Source/Editor/AudioEditor/Private/AssetTypeActions/AssetDefinition_SoundSourceBus.h
+++ b/Engine/Source/Editor/AudioEditor/Private/AssetTypeActions/AssetDefinition_SoundSourceBus.h
@@ -20,7 +20,12 @@ public:
 	virtual TSoftClassPtr<UObject> GetAssetClass() const override { return USoundSourceBus::StaticClass(); }
 	virtual TConstArrayView<FAssetCategoryPath> GetAssetCategories() const override
 	{
-		static const auto Categories = { FAssetCategoryPath(EAssetCategoryPaths::Audio, NSLOCTEXT("AssetTypeActions", "AssetSoundSourceSubMenu", "Source")) };
+		static const auto Categories = 
+			{
+				FAssetCategoryPath(EAssetCategoryPaths::Audio,
+					NSLOCTEXT("AssetDefinition", "AssetDefinitionSoundSourceSubMenu", "Advanced"),
+					FCategoryPath(NSLOCTEXT("AssetDefinition", "AssetDefinition_SoundSourceSubMenuSection", "Routing"), ECategoryMenuType::Section))
+			};
 		return Categories;
 	}
 	virtual TSharedPtr<SWidget> GetThumbnailOverlay(const FAssetData& InAssetData) const override { return nullptr; }
diff --git a/Engine/Source/Editor/AudioEditor/Private/AssetTypeActions/AssetDefinition_SoundSubmix.cpp b/Engine/Source/Editor/AudioEditor/Private/AssetTypeActions/AssetDefinition_SoundSubmix.cpp
new file mode 100644
--- /dev/null
+++ b/Engine/Source/Editor/AudioEditor/Private/AssetTypeActions/AssetDefinition_SoundSubmix.cpp
@@ -0,0 +1,51 @@
+// Copyright Epic Games, Inc. All Rights Reserved.
+
+#include "AssetDefinition_SoundSubmix.h"
+
+#include "AudioEditorModule.h"
+#include "Editor.h"
+#include "SoundSubmixEditor.h"
+#include "Subsystems/AssetEditorSubsystem.h"
+
+EAssetCommandResult UAssetDefinition_SoundSubmix::ActivateAssets(const FAssetActivateArgs& ActivateArgs) const
+{
+	TArray<UObject*> Objects = ActivateArgs.LoadObjects<UObject>();
+	TSet<USoundSubmixBase*> SubmixesToSelect;
+	IAssetEditorInstance* Editor = nullptr;
+	for (UObject* Obj : Objects)
+	{
+		if (USoundSubmixBase* SubmixToSelect = Cast<USoundSubmixBase>(Obj))
+		{
+			if (!Editor)
+			{
+				Editor = GEditor->GetEditorSubsystem<UAssetEditorSubsystem>()->FindEditorForAsset(Obj, true);
+			}
+			SubmixesToSelect.Add(SubmixToSelect);
+		}
+	}
+
+	if (Editor)
+	{
+		static_cast<FSoundSubmixEditor*>(Editor)->SelectSubmixes(SubmixesToSelect);
+		return EAssetCommandResult::Handled;
+	}
+
+	return EAssetCommandResult::Unhandled;
+}
+
+EAssetCommandResult UAssetDefinition_SoundSubmix::OpenAssets(const FAssetOpenArgs& OpenArgs) const
+{
+	TArray<UObject*> Objects = OpenArgs.LoadObjects<UObject>();
+	EToolkitMode::Type Mode = OpenArgs.ToolkitHost.IsValid() ? EToolkitMode::WorldCentric : EToolkitMode::Standalone;
+
+	for (UObject* Obj : Objects)
+	{
+		if (USoundSubmixBase* SoundSubmix = Cast<USoundSubmixBase>(Obj))
+		{
+			IAudioEditorModule* AudioEditorModule = &FModuleManager::LoadModuleChecked<IAudioEditorModule>("AudioEditor");
+			AudioEditorModule->CreateSoundSubmixEditor(Mode, OpenArgs.ToolkitHost, SoundSubmix);
+		}
+	}
+
+	return EAssetCommandResult::Handled;
+}
diff --git a/Engine/Source/Editor/AudioEditor/Private/AssetTypeActions/AssetDefinition_SoundSubmix.h b/Engine/Source/Editor/AudioEditor/Private/AssetTypeActions/AssetDefinition_SoundSubmix.h
new file mode 100644
--- /dev/null
+++ b/Engine/Source/Editor/AudioEditor/Private/AssetTypeActions/AssetDefinition_SoundSubmix.h
@@ -0,0 +1,167 @@
+// Copyright Epic Games, Inc. All Rights Reserved.
+
+#pragma once
+#include "AssetTypeActions/AssetDefinitionDefault_AudioDiffable.h"
+#include "SoundSubmixDefaultColorPalette.h"
+#include "AssetDefinition_SoundSubmix.generated.h"
+
+UCLASS()
+class UAssetDefinition_SoundSubmix : public UAssetDefinitionDefault_AudioDiffable
+{
+	GENERATED_BODY()
+
+public:
+	// UAssetDefinition Begin
+	virtual FText GetAssetDisplayName() const override { return NSLOCTEXT("AssetDefinition", "AssetDefinition_SoundSubmix", "Sound Submix"); }
+	virtual FLinearColor GetAssetColor() const override { return Audio::DefaultSubmixColor; }
+	virtual TSoftClassPtr<UObject> GetAssetClass() const override { return USoundSubmix::StaticClass(); }
+	virtual TConstArrayView<FAssetCategoryPath> GetAssetCategories() const override
+	{
+		static const auto Categories = 
+			{
+				FAssetCategoryPath(EAssetCategoryPaths::Audio,
+					NSLOCTEXT("AssetDefinition", "AssetDefinition_SoundSubmixSubMenu", "Advanced"),
+					FCategoryPath(NSLOCTEXT("AssetDefinition", "AssetDefinition_SoundSubmixSubMenuSection", "Routing"), ECategoryMenuType::Section))
+			};
+		return Categories;
+	}
+	virtual EAssetCommandResult ActivateAssets(const FAssetActivateArgs& ActivateArgs) const override;
+	virtual EAssetCommandResult OpenAssets(const FAssetOpenArgs& OpenArgs) const override;
+	// UAssetDefinition End
+};
+
+UCLASS()
+class UAssetDefinition_SoundfieldSubmix : public UAssetDefinition_SoundSubmix
+{
+	GENERATED_BODY()
+public:
+	// UAssetDefinition Begin
+	virtual FText GetAssetDisplayName() const override { return NSLOCTEXT("AssetDefinition", "AssetDefinition_SoundfieldSubmix", "Soundfield Submix"); }
+	virtual FLinearColor GetAssetColor() const override { return Audio::SoundfieldDefaultSubmixColor; }
+	virtual TSoftClassPtr<UObject> GetAssetClass() const override { return USoundfieldSubmix::StaticClass(); }
+	virtual TConstArrayView<FAssetCategoryPath> GetAssetCategories() const override
+	{
+		static const auto Categories = { FAssetCategoryPath(EAssetCategoryPaths::Audio, NSLOCTEXT("AssetDefinition", "AssetDefinition_SoundfieldSubmixSubMenu", "Legacy")) };
+		return Categories;
+	}
+	// UAssetDefinition End
+};
+
+UCLASS()
+class UAssetDefinition_EndpointSubmix : public UAssetDefinition_SoundSubmix
+{
+	GENERATED_BODY()
+public:
+	// UAssetDefinition Begin
+	virtual FText GetAssetDisplayName() const override { return NSLOCTEXT("AssetDefinition", "AssetDefinition_EndpointSubmix", "Endpoint Submix"); }
+	virtual FLinearColor GetAssetColor() const override { return Audio::EndpointDefaultSubmixColor; }
+	virtual TSoftClassPtr<UObject> GetAssetClass() const override { return UEndpointSubmix::StaticClass(); }
+	virtual TConstArrayView<FAssetCategoryPath> GetAssetCategories() const override
+	{
+		static const auto Categories = { FAssetCategoryPath(EAssetCategoryPaths::Audio, NSLOCTEXT("AssetDefinition", "AssetDefinition_EndpointSubmixSubMenu", "Legacy")) };
+		return Categories;
+	}
+	// UAssetDefinition End
+};
+
+UCLASS()
+class UAssetDefinition_SoundfieldEndpointSubmix : public UAssetDefinition_SoundSubmix
+{
+	GENERATED_BODY()
+public:
+	// UAssetDefinition Begin
+	virtual FText GetAssetDisplayName() const override { return NSLOCTEXT("AssetDefinition", "AssetDefinition_SoundfieldEndpointSubmix", "Soundfield Endpoint Submix"); }
+	virtual FLinearColor GetAssetColor() const override { return Audio::SoundfieldEndpointDefaultSubmixColor; }
+	virtual TSoftClassPtr<UObject> GetAssetClass() const override { return USoundfieldEndpointSubmix::StaticClass(); }
+	virtual TConstArrayView<FAssetCategoryPath> GetAssetCategories() const override
+	{
+		static const auto Categories = { FAssetCategoryPath(EAssetCategoryPaths::Audio, NSLOCTEXT("AssetDefinition", "AssetDefinition_SoundfieldEndpointSubmixSubMenu", "Legacy")) };
+		return Categories;
+	}
+	// UAssetDefinition End
+};
+
+UCLASS()
+class UAssetDefinition_SoundfieldEncodingSettings : public UAssetDefinitionDefault_AudioDiffable
+{
+	GENERATED_BODY()
+public:
+	// UAssetDefinition Begin
+	virtual FText GetAssetDisplayName() const override { return NSLOCTEXT("AssetDefinition", "AssetDefinition_SoundfieldEncodingSettings", "Soundfield Encoding Settings"); }
+	virtual FLinearColor GetAssetColor() const override { return Audio::SoundfieldDefaultSubmixColor; }
+	virtual TSoftClassPtr<UObject> GetAssetClass() const override { return USoundfieldEncodingSettingsBase::StaticClass(); }
+	virtual TConstArrayView<FAssetCategoryPath> GetAssetCategories() const override
+	{
+		static const auto Categories = { FAssetCategoryPath(EAssetCategoryPaths::Audio, NSLOCTEXT("AssetDefinition", "AssetDefinition_SoundfieldEncodingSettingsSubMenu", "Legacy")) };
+		return Categories;
+	}
+	// UAssetDefinition End
+};
+
+UCLASS()
+class UAssetDefinition_SoundfieldEffectSettings : public UAssetDefinitionDefault_AudioDiffable
+{
+	GENERATED_BODY()
+public:
+	// UAssetDefinition Begin
+	virtual FText GetAssetDisplayName() const override { return NSLOCTEXT("AssetDefinition", "AssetDefinition_SoundfieldEffectSettings", "Soundfield Effect Settings"); }
+	virtual FLinearColor GetAssetColor() const override { return Audio::SoundfieldDefaultSubmixColor; }
+	virtual TSoftClassPtr<UObject> GetAssetClass() const override { return USoundfieldEffectSettingsBase::StaticClass(); }
+	virtual TConstArrayView<FAssetCategoryPath> GetAssetCategories() const override
+	{
+		static const auto Categories = { FAssetCategoryPath(EAssetCategoryPaths::Audio, NSLOCTEXT("AssetDefinition", "AssetDefinition_SoundfieldEffectSettingsSubMenu", "Legacy")) };
+		return Categories;
+	}
+	// UAssetDefinition End
+};
+
+UCLASS()
+class UAssetDefinition_SoundfieldEffect : public UAssetDefinitionDefault_AudioDiffable
+{
+	GENERATED_BODY()
+public:
+	// UAssetDefinition Begin
+	virtual FText GetAssetDisplayName() const override { return NSLOCTEXT("AssetDefinition", "AssetDefinition_SoundfieldEffect", "Soundfield Effect"); }
+	virtual FLinearColor GetAssetColor() const override { return Audio::SoundfieldDefaultSubmixColor; }
+	virtual TSoftClassPtr<UObject> GetAssetClass() const override { return USoundfieldEffectBase::StaticClass(); }
+	virtual TConstArrayView<FAssetCategoryPath> GetAssetCategories() const override
+	{
+		static const auto Categories = { FAssetCategoryPath(EAssetCategoryPaths::Audio, NSLOCTEXT("AssetDefinition", "AssetDefinition_SoundfieldEffectSubMenu", "Legacy")) };
+		return Categories;
+	}
+	// UAssetDefinition End
+};
+
+UCLASS()
+class UAssetDefinition_AudioEndpointSettings : public UAssetDefinitionDefault_AudioDiffable
+{
+	GENERATED_BODY()
+public:
+	// UAssetDefinition Begin
+	virtual FText GetAssetDisplayName() const override { return NSLOCTEXT("AssetDefinition", "AssetDefinition_AudioEndpointSettings", "Audio Endpoint Settings"); }
+	virtual FLinearColor GetAssetColor() const override { return Audio::EndpointDefaultSubmixColor; }
+	virtual TSoftClassPtr<UObject> GetAssetClass() const override { return UAudioEndpointSettingsBase::StaticClass(); }
+	virtual TConstArrayView<FAssetCategoryPath> GetAssetCategories() const override
+	{
+		static const auto Categories = { FAssetCategoryPath(EAssetCategoryPaths::Audio, NSLOCTEXT("AssetDefinition", "AssetDefinition_AudioEndpointSettingsSubMenu", "Legacy")) };
+		return Categories;
+	}
+	// UAssetDefinition End
+};
+
+UCLASS()
+class UAssetDefinition_SoundfieldEndpointSettings : public UAssetDefinitionDefault_AudioDiffable
+{
+	GENERATED_BODY()
+public:
+	// UAssetDefinition Begin
+	virtual FText GetAssetDisplayName() const override { return NSLOCTEXT("AssetDefinition", "AssetDefinition_SoundfieldEndpointSettings", "Soundfield Endpoint Settings"); }
+	virtual FLinearColor GetAssetColor() const override { return Audio::EndpointDefaultSubmixColor; }
+	virtual TSoftClassPtr<UObject> GetAssetClass() const override { return USoundfieldEndpointSettingsBase::StaticClass(); }
+	virtual TConstArrayView<FAssetCategoryPath> GetAssetCategories() const override
+	{
+		static const auto Categories = { FAssetCategoryPath(EAssetCategoryPaths::Audio, NSLOCTEXT("AssetDefinition", "AssetDefinition_SoundfieldEndpointSettingsSubMenu", "Legacy")) };
+		return Categories;
+	}
+	// UAssetDefinition End
+};
diff --git a/Engine/Source/Editor/AudioEditor/Private/AssetTypeActions/AssetTypeActions_DialogueWave.cpp b/Engine/Source/Editor/AudioEditor/Private/AssetTypeActions/AssetTypeActions_DialogueWave.cpp
deleted file mode 100644
diff --git a/Engine/Source/Editor/AudioEditor/Private/AssetTypeActions/AssetTypeActions_DialogueWave.h b/Engine/Source/Editor/AudioEditor/Private/AssetTypeActions/AssetTypeActions_DialogueWave.h
deleted file mode 100644
diff --git a/Engine/Source/Editor/AudioEditor/Private/AssetTypeActions/AssetTypeActions_SoundClass.cpp b/Engine/Source/Editor/AudioEditor/Private/AssetTypeActions/AssetTypeActions_SoundClass.cpp
deleted file mode 100644
diff --git a/Engine/Source/Editor/AudioEditor/Private/AssetTypeActions/AssetTypeActions_SoundClass.h b/Engine/Source/Editor/AudioEditor/Private/AssetTypeActions/AssetTypeActions_SoundClass.h
deleted file mode 100644
diff --git a/Engine/Source/Editor/AudioEditor/Private/AssetTypeActions/AssetTypeActions_SoundEffectPreset.cpp b/Engine/Source/Editor/AudioEditor/Private/AssetTypeActions/AssetTypeActions_SoundEffectPreset.cpp
deleted file mode 100644
diff --git a/Engine/Source/Editor/AudioEditor/Private/AssetTypeActions/AssetTypeActions_SoundEffectPreset.h b/Engine/Source/Editor/AudioEditor/Private/AssetTypeActions/AssetTypeActions_SoundEffectPreset.h
deleted file mode 100644
diff --git a/Engine/Source/Editor/AudioEditor/Private/AssetTypeActions/AssetTypeActions_SoundSubmix.cpp b/Engine/Source/Editor/AudioEditor/Private/AssetTypeActions/AssetTypeActions_SoundSubmix.cpp
deleted file mode 100644
diff --git a/Engine/Source/Editor/AudioEditor/Private/AssetTypeActions/AssetTypeActions_SoundSubmix.h b/Engine/Source/Editor/AudioEditor/Private/AssetTypeActions/AssetTypeActions_SoundSubmix.h
deleted file mode 100644
diff --git a/Engine/Source/Editor/AudioEditor/Private/AudioEditorModule.cpp b/Engine/Source/Editor/AudioEditor/Private/AudioEditorModule.cpp
--- a/Engine/Source/Editor/AudioEditor/Private/AudioEditorModule.cpp
+++ b/Engine/Source/Editor/AudioEditor/Private/AudioEditorModule.cpp
@@ -2,11 +2,8 @@
 
 #include "AudioEditorModule.h"
 
+#include "AssetDefinitionRegistry.h"
 #include "AssetToolsModule.h"
-#include "AssetTypeActions/AssetTypeActions_DialogueWave.h"
-#include "AssetTypeActions/AssetTypeActions_SoundClass.h"
-#include "AssetTypeActions/AssetTypeActions_SoundEffectPreset.h"
-#include "AssetTypeActions/AssetTypeActions_SoundSubmix.h"
 #include "AudioEditorSettings.h"
 #include "ClassTemplateEditorSubsystem.h"
 #include "ComponentAssetBroker.h"
@@ -18,6 +15,7 @@
 #include "Modules/ModuleManager.h"
 #include "PropertyEditorModule.h"
 #include "Sound/AudioSettings.h"
+#include "Sound/DialogueWave.h"
 #include "Sound/SoundClass.h"
 #include "Sound/SoundCue.h"
 #include "Sound/SoundEffectPreset.h"
@@ -39,7 +37,7 @@
 #include "UObject/UObjectIterator.h"
 #include "Utils.h"
 #include "WidgetBlueprint.h"
-
+#include "AssetTypeActions/AssetDefinition_SoundEffectPreset.h"
 const FName AudioEditorAppIdentifier("AudioEditorApp");
 
 DEFINE_LOG_CATEGORY(LogAudioEditor);
@@ -221,38 +219,18 @@ public:
 		}
 	}
 
+	PRAGMA_DISABLE_DEPRECATION_WARNINGS
 	virtual void RegisterAssetActions() override
-	{
-		LLM_SCOPE(ELLMTag::AudioMisc);
-		// Register the audio editor asset type actions
-		IAssetTools& AssetTools = FModuleManager::LoadModuleChecked<FAssetToolsModule>("AssetTools").Get();
-
-		AssetTools.RegisterAssetTypeActions(MakeShared<FAssetTypeActions_DialogueWave>());
-		AssetTools.RegisterAssetTypeActions(MakeShared<FAssetTypeActions_SoundClass>());
-	}
+	{}
+	PRAGMA_ENABLE_DEPRECATION_WARNINGS
 
+	PRAGMA_DISABLE_DEPRECATION_WARNINGS
 	virtual void RegisterAudioMixerAssetActions() override
-	{
-		IAssetTools& AssetTools = FModuleManager::LoadModuleChecked<FAssetToolsModule>("AssetTools").Get();
-		AssetTools.RegisterAssetTypeActions(MakeShared<FAssetTypeActions_SoundSubmix>());
-		AssetTools.RegisterAssetTypeActions(MakeShared<FAssetTypeActions_SoundfieldSubmix>());
-		AssetTools.RegisterAssetTypeActions(MakeShared<FAssetTypeActions_EndpointSubmix>());
-		AssetTools.RegisterAssetTypeActions(MakeShared<FAssetTypeActions_SoundfieldEndpointSubmix>());
-		AssetTools.RegisterAssetTypeActions(MakeShared<FAssetTypeActions_SoundfieldEncodingSettings>());
-		AssetTools.RegisterAssetTypeActions(MakeShared<FAssetTypeActions_SoundfieldEffectSettings>());
-		AssetTools.RegisterAssetTypeActions(MakeShared<FAssetTypeActions_SoundfieldEffect>());
-		AssetTools.RegisterAssetTypeActions(MakeShared<FAssetTypeActions_AudioEndpointSettings>());
-		AssetTools.RegisterAssetTypeActions(MakeShared<FAssetTypeActions_SoundfieldEndpointSettings>());
-		AssetTools.RegisterAssetTypeActions(MakeShared<FAssetTypeActions_SoundEffectSubmixPreset>());
-		AssetTools.RegisterAssetTypeActions(MakeShared<FAssetTypeActions_SoundEffectSourcePreset>());
-		AssetTools.RegisterAssetTypeActions(MakeShared<FAssetTypeActions_SoundEffectSourcePresetChain>());
-	}
+	{}
+	PRAGMA_ENABLE_DEPRECATION_WARNINGS
 
 	virtual void RegisterEffectPresetAssetActions() override
 	{
-		// Register the audio editor asset type actions
-		IAssetTools& AssetTools = FModuleManager::LoadModuleChecked<FAssetToolsModule>("AssetTools").Get();
-
 		// Look for any sound effect presets to register
 		for (TObjectIterator<UClass> It; It; ++It)
 		{
@@ -270,7 +248,12 @@ public:
 				if (!RegisteredActions.Contains(EffectPreset) && EffectPreset->HasAssetActions())
 				{
 					RegisteredActions.Add(EffectPreset);
-					AssetTools.RegisterAssetTypeActions(MakeShared<FAssetTypeActions_SoundEffectPreset>(EffectPreset));
+					UAssetDefinition_SoundEffectPresetDynamic* Definition = NewObject<UAssetDefinition_SoundEffectPresetDynamic>(
+						GetTransientPackage(),
+						MakeUniqueObjectName(GetTransientPackage(), UAssetDefinition_SoundEffectPresetDynamic::StaticClass())
+					);
+					Definition->Initialize(TSubclassOf<USoundEffectPreset>(ChildClass));
+					UAssetDefinitionRegistry::Get()->RegisterAssetDefinition(Definition);
 				}
 			}
 		}
diff --git a/Engine/Source/Editor/AudioEditor/Public/AudioEditorModule.h b/Engine/Source/Editor/AudioEditor/Public/AudioEditorModule.h
--- a/Engine/Source/Editor/AudioEditor/Public/AudioEditorModule.h
+++ b/Engine/Source/Editor/AudioEditor/Public/AudioEditorModule.h
@@ -54,9 +54,11 @@ class IAudioEditorModule :	public IModuleInterface
 {
 public:
 	/** Registers audio editor asset actions. */
+	UE_DEPRECATED(5.8, "Asset action where converted to asset definition, there is no need to register them manually anymore")
 	virtual void RegisterAssetActions() = 0;
 
 	/** Registers audio editor asset actions specific to audio mixer functionality. */
+	UE_DEPRECATED(5.8, "Asset action where converted to asset definition, there is no need to register them manually anymore")
 	virtual void RegisterAudioMixerAssetActions() = 0;
 
 	/** Registers effect preset asset actions. */
diff --git a/Engine/Source/Runtime/Engine/Private/AudioDeviceManager.cpp b/Engine/Source/Runtime/Engine/Private/AudioDeviceManager.cpp
--- a/Engine/Source/Runtime/Engine/Private/AudioDeviceManager.cpp
+++ b/Engine/Source/Runtime/Engine/Private/AudioDeviceManager.cpp
@@ -439,7 +439,6 @@ bool FAudioDeviceManager::InitializeManager()
 
 #if WITH_EDITOR
 			IAudioEditorModule* AudioEditorModule = &FModuleManager::LoadModuleChecked<IAudioEditorModule>("AudioEditor");
-			AudioEditorModule->RegisterAudioMixerAssetActions();
 			AudioEditorModule->RegisterEffectPresetAssetActions();
 #endif
 
diff --git a/Engine/Source/Runtime/Launch/Private/LaunchEngineLoop.cpp b/Engine/Source/Runtime/Launch/Private/LaunchEngineLoop.cpp
--- a/Engine/Source/Runtime/Launch/Private/LaunchEngineLoop.cpp
+++ b/Engine/Source/Runtime/Launch/Private/LaunchEngineLoop.cpp
@@ -4602,8 +4602,10 @@ bool FEngineLoop::LoadStartupCoreModules()
 #endif
 
 #if WITH_EDITOR
+	PRAGMA_DISABLE_DEPRECATION_WARNINGS
 	IAudioEditorModule* AudioEditorModule = &FModuleManager::LoadModuleChecked<IAudioEditorModule>("AudioEditor");
 	AudioEditorModule->RegisterAssetActions();
+	PRAGMA_ENABLE_DEPRECATION_WARNINGS
 #endif // WITH_EDITOR
 
 #if (WITH_EDITOR && !(UE_BUILD_SHIPPING || UE_BUILD_TEST)) // todo: revisit
